Changed the default so game content can access default game and engine plugins without an explicit dependency (matching the reference behavior; any plugin that matches a more specific rule doesn't get this automatically, and the behavior can be disabled independently for engine or game plugins via bCanProjectAccessDefaultPlugins)

diff --git a/Engine/Plugins/Editor/AssetReferenceRestrictions/Source/AssetReferenceRestrictions/Private/AssetReferencingDomains.cpp b/Engine/Plugins/Editor/AssetReferenceRestrictions/Source/AssetReferenceRestrictions/Private/AssetReferencingDomains.cpp
--- a/Engine/Plugins/Editor/AssetReferenceRestrictions/Source/AssetReferenceRestrictions/Private/AssetReferencingDomains.cpp
+++ b/Engine/Plugins/Editor/AssetReferenceRestrictions/Source/AssetReferenceRestrictions/Private/AssetReferencingDomains.cpp
@@ -116,6 +116,7 @@ FDomainDatabase::~FDomainDatabase()
 	IPluginManager& PluginManager = IPluginManager::Get();
 	PluginManager.OnNewPluginCreated().RemoveAll(this);
 	PluginManager.OnNewPluginMounted().RemoveAll(this);
+	PluginManager.OnPluginEdited().RemoveAll(this);
 }
 
 void FDomainDatabase::RebuildFromScratch()
@@ -217,6 +218,8 @@ void FDomainDatabase::Init()
 	IPluginManager& PluginManager = IPluginManager::Get();
 	PluginManager.OnNewPluginCreated().AddRaw(this, &FDomainDatabase::OnPluginCreatedOrMounted);
 	PluginManager.OnNewPluginMounted().AddRaw(this, &FDomainDatabase::OnPluginCreatedOrMounted);
+	PluginManager.OnPluginEdited().AddRaw(this, &FDomainDatabase::OnPluginCreatedOrMounted);
+	
 }
 
 void FDomainDatabase::MarkDirty()
@@ -326,7 +329,7 @@ void FDomainDatabase::BuildDomainFromPlugin(TSharedRef<IPlugin> Plugin)
 	{
 		AddDomainVisibilityList(Domain, DomainSettings.DefaultRule.CanReferenceTheseDomains);
 
-		if ((Plugin->GetLoadedFrom() == EPluginLoadedFrom::Project) && Settings->bCanProjectAccessDefaultProjectPlugins)
+		if (DomainSettings.DefaultRule.bCanProjectAccessDefaultPlugins)
 		{
 			GameDomain->DomainsVisibleFromHere.Add(Domain);
 		}
diff --git a/Engine/Plugins/Editor/AssetReferenceRestrictions/Source/AssetReferenceRestrictions/Private/AssetReferencingPolicySettings.h b/Engine/Plugins/Editor/AssetReferenceRestrictions/Source/AssetReferenceRestrictions/Private/AssetReferencingPolicySettings.h
--- a/Engine/Plugins/Editor/AssetReferenceRestrictions/Source/AssetReferenceRestrictions/Private/AssetReferencingPolicySettings.h
+++ b/Engine/Plugins/Editor/AssetReferenceRestrictions/Source/AssetReferenceRestrictions/Private/AssetReferencingPolicySettings.h
@@ -16,6 +16,10 @@ struct FARPDefaultPluginDomainRules
 	// (EngineContent is always visible, as is content from other plugins that are explicitly referenced)
 	UPROPERTY(EditAnywhere, Category=Settings, meta=(GetOptions="GetListOfDomains_NoPluginsOrEngine"))
 	TArray<FString> CanReferenceTheseDomains;
+
+	// Can content in the ProjectContent domain access plugin content automatically (for plugins that don't match a specific rule)?
+	UPROPERTY(config, EditAnywhere, Category = "Project Plugins")
+	bool bCanProjectAccessDefaultPlugins = true;
 };
 
 UENUM()
@@ -30,10 +34,6 @@ struct FARPDomainDefinitionForMatchingPlugins
 {
 	GENERATED_BODY()
 
-	// Name of this rule template (used as a merge key in the ini)
-// 	UPROPERTY(EditAnywhere, Category=Settings)
-// 	FString TemplateName;
-
 	// The display name of this rule (used in error message when attempting to reference content incorrectly)
 	// The token {0} will be replaced with the plugin name
 	UPROPERTY(EditAnywhere, Category=Settings)
@@ -167,15 +167,14 @@ class UAssetReferencingPolicySettings : public UDeveloperSettings
 	GENERATED_BODY()
 
 public:
-	UPROPERTY(config, EditAnywhere, Category="Engine Plugins")//, meta=(ShowOnlyInnerProperties)
+	// Settings/rules for engine plugins
+	UPROPERTY(config, EditAnywhere, Category="Engine Plugins")
 	FARPDomainSettingsForPlugins EnginePlugins;
 
-	UPROPERTY(config, EditAnywhere, Category="Project Plugins")//, meta=(ShowOnlyInnerProperties)
+	// Settings/rules for project plugins
+	UPROPERTY(config, EditAnywhere, Category="Project Plugins")
 	FARPDomainSettingsForPlugins ProjectPlugins;
 
-	UPROPERTY(config, EditAnywhere, Category="Project Plugins")//, meta=(ShowOnlyInnerProperties)
-	bool bCanProjectAccessDefaultProjectPlugins = true;
-
 	// The default rules for project content (if a more specific rule doesn't apply)
 	UPROPERTY(config, EditAnywhere, Category="Project Content")
 	FARPDefaultProjectDomainRules DefaultProjectContentRule;
