Fixed the asset filter failing to determine a domain when using the class picker (empty asset datas for the (none) option) and in a brand new unsaved map (stuff in /Temp/)

diff --git a/Engine/Plugins/Editor/AssetReferenceRestrictions/Source/AssetReferenceRestrictions/Private/AssetReferencingDomains.cpp b/Engine/Plugins/Editor/AssetReferenceRestrictions/Source/AssetReferenceRestrictions/Private/AssetReferencingDomains.cpp
--- a/Engine/Plugins/Editor/AssetReferenceRestrictions/Source/AssetReferenceRestrictions/Private/AssetReferencingDomains.cpp
+++ b/Engine/Plugins/Editor/AssetReferenceRestrictions/Source/AssetReferenceRestrictions/Private/AssetReferencingDomains.cpp
@@ -127,14 +127,23 @@ void FDomainDatabase::RebuildFromScratch()
 	DomainNameMap.Reset();
 	PathMap.Reset();
 	EngineDomain.Reset();
+	TempDomain.Reset();
 	GameDomain.Reset();
 	DomainsDefinedByPlugins.Reset();
 
 	// Create the built-in domains
+	//@TODO: Would be good to get the path roots directly from FLongPackagePathsSingleton but it's private
 	EngineDomain = FindOrAddDomainByName(UAssetReferencingPolicySettings::EngineDomainName);
 	EngineDomain->UserFacingDomainName = LOCTEXT("EngineDomain", "EngineContent");
 	EngineDomain->DomainRootPaths.Add(TEXT("/Engine/"));
 
+	TempDomain = FindOrAddDomainByName(UAssetReferencingPolicySettings::TempDomainName);
+	TempDomain->UserFacingDomainName = LOCTEXT("TempDomain", "Temp");
+	TempDomain->bCanSeeEverything = true;
+	TempDomain->DomainRootPaths.Add(TEXT("/Temp/"));
+	TempDomain->DomainRootPaths.Add(TEXT("/Extra/"));
+	TempDomain->DomainRootPaths.Add(TEXT("/Memory/"));
+
 	GameDomain = FindOrAddDomainByName(UAssetReferencingPolicySettings::GameDomainName);
 	GameDomain->UserFacingDomainName = LOCTEXT("GameDomain", "ProjectContent");  	//@TODO: FText::AsCultureInvariant(GetDefault<UGeneralProjectSettings>()->ProjectName;);
 	GameDomain->DomainRootPaths.Add(TEXT("/Game/"));
diff --git a/Engine/Plugins/Editor/AssetReferenceRestrictions/Source/AssetReferenceRestrictions/Private/AssetReferencingDomains.h b/Engine/Plugins/Editor/AssetReferenceRestrictions/Source/AssetReferenceRestrictions/Private/AssetReferencingDomains.h
--- a/Engine/Plugins/Editor/AssetReferenceRestrictions/Source/AssetReferenceRestrictions/Private/AssetReferencingDomains.h
+++ b/Engine/Plugins/Editor/AssetReferenceRestrictions/Source/AssetReferenceRestrictions/Private/AssetReferencingDomains.h
@@ -28,6 +28,8 @@ struct FDomainData final : public TSharedFromThis<FDomainData>
 
 	// The domains that are visible from here (if bCanSeeEverything is true, then literally everything is visible from here)
 	TSet<TSharedPtr<FDomainData>> DomainsVisibleFromHere;
+
+	// Can we see everything?
 	bool bCanSeeEverything = false;
 
 	bool IsValid() const
@@ -81,6 +83,10 @@ private:
 	// The engine content domain
 	TSharedPtr<FDomainData> EngineDomain;
 
+	// Used for various 'special' mount points like /Temp/, /Memory/, and /Extra/
+	// Not visible as a domain for other domains to see, and can see everything
+	TSharedPtr<FDomainData> TempDomain;
+
 	// The game content domain
 	TSharedPtr<FDomainData> GameDomain;
 
diff --git a/Engine/Plugins/Editor/AssetReferenceRestrictions/Source/AssetReferenceRestrictions/Private/AssetReferencingPolicySettings.cpp b/Engine/Plugins/Editor/AssetReferenceRestrictions/Source/AssetReferenceRestrictions/Private/AssetReferencingPolicySettings.cpp
--- a/Engine/Plugins/Editor/AssetReferenceRestrictions/Source/AssetReferenceRestrictions/Private/AssetReferencingPolicySettings.cpp
+++ b/Engine/Plugins/Editor/AssetReferenceRestrictions/Source/AssetReferenceRestrictions/Private/AssetReferencingPolicySettings.cpp
@@ -7,6 +7,7 @@
 
 const FString UAssetReferencingPolicySettings::EngineDomainName(TEXT("EngineContent"));
 const FString UAssetReferencingPolicySettings::GameDomainName(TEXT("ProjectContent"));
+const FString UAssetReferencingPolicySettings::TempDomainName(TEXT("Temp"));
 
 UAssetReferencingPolicySettings::UAssetReferencingPolicySettings(const FObjectInitializer& ObjectInitializer)
 	: Super(ObjectInitializer)
@@ -38,7 +39,7 @@ TArray<FString> UAssetReferencingPolicySettings::GetListOfDomains(bool bAllowPlu
 
 	if (bAllowGame)
 	{
-		Result.Add(TEXT("ProjectContent"));
+		Result.Add(GameDomainName);
 	}
 
 	for (const FARPDomainDefinitionByContentRoot& DomainDef : AdditionalDomains)
@@ -55,7 +56,7 @@ TArray<FString> UAssetReferencingPolicySettings::GetListOfDomains(bool bAllowPlu
 
 	if (bAllowEngine)
 	{
-		Result.Add(TEXT("EngineContent"));
+		Result.Add(EngineDomainName);
 	}
 
 	return Result;
diff --git a/Engine/Plugins/Editor/AssetReferenceRestrictions/Source/AssetReferenceRestrictions/Private/AssetReferencingPolicySettings.h b/Engine/Plugins/Editor/AssetReferenceRestrictions/Source/AssetReferenceRestrictions/Private/AssetReferencingPolicySettings.h
--- a/Engine/Plugins/Editor/AssetReferenceRestrictions/Source/AssetReferenceRestrictions/Private/AssetReferencingPolicySettings.h
+++ b/Engine/Plugins/Editor/AssetReferenceRestrictions/Source/AssetReferenceRestrictions/Private/AssetReferencingPolicySettings.h
@@ -183,8 +183,10 @@ public:
 	UPROPERTY(config, EditAnywhere, Category="Project Content", meta=(TitleProperty=DomainName))
 	TArray<FARPDomainDefinitionByContentRoot> AdditionalDomains;
 
+	// The names of the project, special system mount, and game domains
 	static const FString EngineDomainName;
 	static const FString GameDomainName;
+	static const FString TempDomainName;
 
 public:
 	//~UObject interface
diff --git a/Engine/Plugins/Editor/AssetReferenceRestrictions/Source/AssetReferenceRestrictions/Private/DomainAssetReferenceFilter.cpp b/Engine/Plugins/Editor/AssetReferenceRestrictions/Source/AssetReferenceRestrictions/Private/DomainAssetReferenceFilter.cpp
--- a/Engine/Plugins/Editor/AssetReferenceRestrictions/Source/AssetReferenceRestrictions/Private/DomainAssetReferenceFilter.cpp
+++ b/Engine/Plugins/Editor/AssetReferenceRestrictions/Source/AssetReferenceRestrictions/Private/DomainAssetReferenceFilter.cpp
@@ -137,14 +137,19 @@ bool FDomainAssetReferenceFilter::PassesFilterImpl(const FAssetData& AssetData,
 {
 	QUICK_SCOPE_CYCLE_COUNTER(FDomainAssetReferenceFilter_PassesFilterImpl);
 
-	const FString ReferencedAssetPath = AssetData.PackageName.ToString();
+	// An empty/null asset reference is always valid (e.g., the (none) option in the class picker)
+	if (!AssetData.IsValid())
+	{
+		return true;
+	}
 
 	TSharedPtr<FDomainData> AssetDomain = DomainDB->FindDomainFromAssetData(AssetData);
 
-	// This ensure should never fail; every asset should belong to engine, game, or a plugin
-	// that is allowed to contain content, and we create domains for all of those.
+	// This ensure should never fail; every valid asset should belong to engine, special system mount,
+	// game, or a plugin that is allowed to contain content, and we create domains for all of those.
 	// The only exception anticipated is if somehow an editor with an asset picker is pointed
 	// towards something living in the transient package, but that's not expected!
+	// If it gets hit, transient can be added to the special system mount 'Temp' domain
 	if (ensureMsgf(AssetDomain.IsValid(), TEXT("Asset %s didn't match any domain"), *AssetData.ObjectPath.ToString()))
 	{
 		// Check the referencing domains
