Added a temporary mechanism to allow the asset validator plugin to be enabled but rendered inert by config variables, and enabled the plugin by default when the game feature plugin is enabled

diff --git a/Engine/Plugins/Editor/AssetReferenceRestrictions/Source/AssetReferenceRestrictions/Private/AssetReferencingPolicySettings.h b/Engine/Plugins/Editor/AssetReferenceRestrictions/Source/AssetReferenceRestrictions/Private/AssetReferencingPolicySettings.h
--- a/Engine/Plugins/Editor/AssetReferenceRestrictions/Source/AssetReferenceRestrictions/Private/AssetReferencingPolicySettings.h
+++ b/Engine/Plugins/Editor/AssetReferenceRestrictions/Source/AssetReferenceRestrictions/Private/AssetReferencingPolicySettings.h
@@ -183,6 +183,16 @@ public:
 	UPROPERTY(config, EditAnywhere, Category="Project Content", meta=(TitleProperty=DomainName))
 	TArray<FARPDomainDefinitionByContentRoot> AdditionalDomains;
 
+	// This is a temporary flags for testing and should not be disabled under normal circumstances
+	// (it will be removed by 5.0)
+	UPROPERTY(config)
+	bool bEnableAssetFiltering_TEMP = true;
+
+	// This is a temporary flag for testing and should not be disabled under normal circumstances
+	// (it will be removed by 5.0)
+	UPROPERTY(config)
+	bool bEnableAssetValidator_TEMP = true;
+
 	// The names of the project, special system mount, and game domains
 	static const FString EngineDomainName;
 	static const FString GameDomainName;
diff --git a/Engine/Plugins/Editor/AssetReferenceRestrictions/Source/AssetReferenceRestrictions/Private/AssetReferencingPolicySubsystem.cpp b/Engine/Plugins/Editor/AssetReferenceRestrictions/Source/AssetReferenceRestrictions/Private/AssetReferencingPolicySubsystem.cpp
--- a/Engine/Plugins/Editor/AssetReferenceRestrictions/Source/AssetReferenceRestrictions/Private/AssetReferencingPolicySubsystem.cpp
+++ b/Engine/Plugins/Editor/AssetReferenceRestrictions/Source/AssetReferenceRestrictions/Private/AssetReferencingPolicySubsystem.cpp
@@ -8,6 +8,11 @@
 
 #define LOCTEXT_NAMESPACE "AssetReferencingPolicy"
 
+bool UAssetReferencingPolicySubsystem::ShouldCreateSubsystem(UObject* Outer) const
+{
+	return GetDefault<UAssetReferencingPolicySettings>()->bEnableAssetFiltering_TEMP;
+}
+
 void UAssetReferencingPolicySubsystem::Initialize(FSubsystemCollectionBase& Collection)
 {
 	check(GEditor);
diff --git a/Engine/Plugins/Editor/AssetReferenceRestrictions/Source/AssetReferenceRestrictions/Private/AssetReferencingPolicySubsystem.h b/Engine/Plugins/Editor/AssetReferenceRestrictions/Source/AssetReferenceRestrictions/Private/AssetReferencingPolicySubsystem.h
--- a/Engine/Plugins/Editor/AssetReferenceRestrictions/Source/AssetReferenceRestrictions/Private/AssetReferencingPolicySubsystem.h
+++ b/Engine/Plugins/Editor/AssetReferenceRestrictions/Source/AssetReferenceRestrictions/Private/AssetReferencingPolicySubsystem.h
@@ -19,6 +19,7 @@ class UAssetReferencingPolicySubsystem : public UEditorSubsystem
 
 public:
 	//~UEditorSubsystem interface
+	virtual bool ShouldCreateSubsystem(UObject* Outer) const override;
 	virtual void Initialize(FSubsystemCollectionBase& Collection) override;
 	virtual void Deinitialize() override;
 	//~End of UEditorSubsystem interface
diff --git a/Engine/Plugins/Editor/AssetReferenceRestrictions/Source/AssetReferenceRestrictions/Private/AssetValidator_AssetReferenceRestrictions.cpp b/Engine/Plugins/Editor/AssetReferenceRestrictions/Source/AssetReferenceRestrictions/Private/AssetValidator_AssetReferenceRestrictions.cpp
--- a/Engine/Plugins/Editor/AssetReferenceRestrictions/Source/AssetReferenceRestrictions/Private/AssetValidator_AssetReferenceRestrictions.cpp
+++ b/Engine/Plugins/Editor/AssetReferenceRestrictions/Source/AssetReferenceRestrictions/Private/AssetValidator_AssetReferenceRestrictions.cpp
@@ -8,6 +8,7 @@
 #include "Editor/UnrealEdEngine.h"
 #include "AssetReferencingPolicySubsystem.h"
 #include "AssetReferencingDomains.h"
+#include "AssetReferencingPolicySettings.h"
 
 #define LOCTEXT_NAMESPACE "AssetReferencingPolicy"
 
@@ -18,6 +19,11 @@ UAssetValidator_AssetReferenceRestrictions::UAssetValidator_AssetReferenceRestri
 
 bool UAssetValidator_AssetReferenceRestrictions::CanValidateAsset_Implementation(UObject* InAsset) const
 {
+	if (!GetDefault<UAssetReferencingPolicySettings>()->bEnableAssetValidator_TEMP)
+	{
+		return false;
+	}
+
 	if (InAsset)
 	{
 		const FAssetData AssetData(InAsset);
@@ -50,16 +56,6 @@ EDataValidationResult UAssetValidator_AssetReferenceRestrictions::ValidateLoaded
 		const FString SoftDependencyStr = SoftDependency.ToString();
 		if (!FPackageName::IsScriptPackage(SoftDependencyStr))
 		{
-
-//@TODO: Probably not needed anymore?
-#if 0
-			FString UncookedFolderName;
-			if (IsInUncookedFolder(SoftDependencyStr, &UncookedFolderName))
-			{
-				AssetFails(InAsset, FText::Format(LOCTEXT("IllegalReference_SoftRef", "Illegally soft references {0} asset {1}"), FText::FromString(UncookedFolderName), FText::FromString(SoftDependencyStr)), ValidationErrors);
-			}
-#endif
-
 			TArray<FAssetData> DependencyAssets;
 			AssetRegistry.GetAssetsByPackageName(SoftDependency, DependencyAssets, true);
 			if (DependencyAssets.Num() == 0)
diff --git a/Engine/Plugins/Experimental/GameFeatures/GameFeatures.uplugin b/Engine/Plugins/Experimental/GameFeatures/GameFeatures.uplugin
--- a/Engine/Plugins/Experimental/GameFeatures/GameFeatures.uplugin
+++ b/Engine/Plugins/Experimental/GameFeatures/GameFeatures.uplugin
@@ -33,7 +33,11 @@
 		},
 		{
 			"Name": "DataRegistry",
-			"Enabled":  true
+			"Enabled": true
+		},
+		{
+			"Name": "AssetReferenceRestrictions",
+			"Enabled": true
 		}
 	]
 }
\ No newline at end of file
