Added an option to treat default* engine or project plugins like 'engine' content so they can be seen by anyone without requiring an explicit dependency (bCanBeSeenByOtherDomainsWithoutDependency)
- This is enabled by default for engine plugins, e.g., Niagara, but not enabled by default for project plugins
- This option might be deprecated/replaced with a different mechanism in the future, but is needed for now to match existing behavior in some games

* default = ones that don't match a more specific template

diff --git a/Engine/Plugins/Editor/AssetReferenceRestrictions/Source/AssetReferenceRestrictions/Private/AssetReferencingDomains.cpp b/Engine/Plugins/Editor/AssetReferenceRestrictions/Source/AssetReferenceRestrictions/Private/AssetReferencingDomains.cpp
--- a/Engine/Plugins/Editor/AssetReferenceRestrictions/Source/AssetReferenceRestrictions/Private/AssetReferencingDomains.cpp
+++ b/Engine/Plugins/Editor/AssetReferenceRestrictions/Source/AssetReferenceRestrictions/Private/AssetReferencingDomains.cpp
@@ -338,7 +338,11 @@ void FDomainDatabase::BuildDomainFromPlugin(TSharedRef<IPlugin> Plugin)
 	{
 		AddDomainVisibilityList(Domain, DomainSettings.DefaultRule.CanReferenceTheseDomains);
 
-		if (DomainSettings.DefaultRule.bCanProjectAccessDefaultPlugins)
+		if (DomainSettings.DefaultRule.bCanBeSeenByOtherDomainsWithoutDependency)
+		{
+			Domain->bCanBeSeenByEverything = true;
+		}
+		else if (DomainSettings.DefaultRule.bCanProjectAccessThesePlugins)
 		{
 			GameDomain->DomainsVisibleFromHere.Add(Domain);
 		}
@@ -413,7 +417,7 @@ TTuple<bool, FText> FDomainDatabase::CanDomainsSeeEachOther(TSharedPtr<FDomainDa
 	check(Referencee.IsValid());
 	check(Referencer.IsValid());
 
-	if ((Referencee == Referencer) || Referencer->DomainsVisibleFromHere.Contains(Referencee))
+	if ((Referencee == Referencer) || Referencer->DomainsVisibleFromHere.Contains(Referencee) || Referencee->bCanBeSeenByEverything)
 	{
 		return MakeTuple(true, FText::GetEmpty());
 	}
diff --git a/Engine/Plugins/Editor/AssetReferenceRestrictions/Source/AssetReferenceRestrictions/Private/AssetReferencingDomains.h b/Engine/Plugins/Editor/AssetReferenceRestrictions/Source/AssetReferenceRestrictions/Private/AssetReferencingDomains.h
--- a/Engine/Plugins/Editor/AssetReferenceRestrictions/Source/AssetReferenceRestrictions/Private/AssetReferencingDomains.h
+++ b/Engine/Plugins/Editor/AssetReferenceRestrictions/Source/AssetReferenceRestrictions/Private/AssetReferencingDomains.h
@@ -32,6 +32,9 @@ struct FDomainData final : public TSharedFromThis<FDomainData>
 	// Can we see everything?
 	bool bCanSeeEverything = false;
 
+	// Can we be seen by everything?
+	bool bCanBeSeenByEverything = false;
+
 	bool IsValid() const
 	{
 		return DomainRootPaths.Num() > 0;
diff --git a/Engine/Plugins/Editor/AssetReferenceRestrictions/Source/AssetReferenceRestrictions/Private/AssetReferencingPolicySettings.cpp b/Engine/Plugins/Editor/AssetReferenceRestrictions/Source/AssetReferenceRestrictions/Private/AssetReferencingPolicySettings.cpp
--- a/Engine/Plugins/Editor/AssetReferenceRestrictions/Source/AssetReferenceRestrictions/Private/AssetReferencingPolicySettings.cpp
+++ b/Engine/Plugins/Editor/AssetReferenceRestrictions/Source/AssetReferenceRestrictions/Private/AssetReferencingPolicySettings.cpp
@@ -19,6 +19,8 @@ UAssetReferencingPolicySettings::UAssetReferencingPolicySettings(const FObjectIn
 	GameFeatureRule.PluginPathPrefix = TEXT("/GameFeatures/");
 	GameFeatureRule.CanReferenceTheseDomains.Add(GameDomainName);
 
+	EnginePlugins.DefaultRule.bCanBeSeenByOtherDomainsWithoutDependency = true;
+
 	ProjectPlugins.DefaultRule.CanReferenceTheseDomains.Add(GameDomainName);
 }
 
diff --git a/Engine/Plugins/Editor/AssetReferenceRestrictions/Source/AssetReferenceRestrictions/Private/AssetReferencingPolicySettings.h b/Engine/Plugins/Editor/AssetReferenceRestrictions/Source/AssetReferenceRestrictions/Private/AssetReferencingPolicySettings.h
--- a/Engine/Plugins/Editor/AssetReferenceRestrictions/Source/AssetReferenceRestrictions/Private/AssetReferencingPolicySettings.h
+++ b/Engine/Plugins/Editor/AssetReferenceRestrictions/Source/AssetReferenceRestrictions/Private/AssetReferencingPolicySettings.h
@@ -18,8 +18,13 @@ struct FARPDefaultPluginDomainRules
 	TArray<FString> CanReferenceTheseDomains;
 
 	// Can content in the ProjectContent domain access plugin content automatically (for plugins that don't match a specific rule)?
-	UPROPERTY(config, EditAnywhere, Category = "Project Plugins")
-	bool bCanProjectAccessDefaultPlugins = true;
+	UPROPERTY(config, EditAnywhere, Category=Settings)
+	bool bCanProjectAccessThesePlugins = true;
+
+	// Can content in other domains access plugin content automatically (for plugins that don't match a specific rule)?
+	// Note: This rule may be deprecated in the future!
+	UPROPERTY(config, EditAnywhere, Category=Settings)
+	bool bCanBeSeenByOtherDomainsWithoutDependency = false;
 };
 
 UENUM()
