Removed the temporary cross-checker from the asset reference restrictions plugin

diff --git a/Engine/Plugins/Editor/AssetReferenceRestrictions/Source/AssetReferenceRestrictions/Private/DomainAssetReferenceFilter.cpp b/Engine/Plugins/Editor/AssetReferenceRestrictions/Source/AssetReferenceRestrictions/Private/DomainAssetReferenceFilter.cpp
--- a/Engine/Plugins/Editor/AssetReferenceRestrictions/Source/AssetReferenceRestrictions/Private/DomainAssetReferenceFilter.cpp
+++ b/Engine/Plugins/Editor/AssetReferenceRestrictions/Source/AssetReferenceRestrictions/Private/DomainAssetReferenceFilter.cpp
@@ -8,10 +8,6 @@
 #include "AssetRegistryModule.h"
 #include "AssetReferencingDomains.h"
 
-#if UE_ASSET_DOMAIN_FILTERING_CROSS_CHECK_RESULTS
-#include "RefAssetReferenceFilter.h"
-#endif
-
 #define LOCTEXT_NAMESPACE "AssetReferencingPolicy"
 
 TArray<FDomainAssetReferenceFilter*> FDomainAssetReferenceFilter::FilterInstances;
@@ -20,10 +16,6 @@ FDomainAssetReferenceFilter::FDomainAssetReferenceFilter(const FAssetReferenceFi
 	: IAssetReferenceFilter()
 	, DomainDB(InDomainDB)
 {
-#if UE_ASSET_DOMAIN_FILTERING_CROSS_CHECK_RESULTS
-	CrossChecker = MakeShareable(new FRefAssetReferenceFilter(Context));
-#endif
-
 	OriginalReferencingAssets = Context.ReferencingAssets;
 
 	DetermineReferencingDomain();
@@ -97,27 +89,6 @@ bool FDomainAssetReferenceFilter::PassesFilter(const FAssetData& AssetData, FTex
 	FText MyReason;
 	bool bMyResult = PassesFilterImpl(AssetData, /*out*/ MyReason);
 
-#if UE_ASSET_DOMAIN_FILTERING_CROSS_CHECK_RESULTS
-	FText OldReason;
-	bool bOldResult = CrossChecker->PassesFilter(AssetData, /*out*/ &OldReason);
-
-	if ((bMyResult != bOldResult) || (MyReason.ToString() != OldReason.ToString()))
-	{
-		UE_LOG(LogInit, Error, TEXT("QQQ Mismatch for asset %s\n\tMe:%d\n\tOld:%d\n\tMyReason: %s\n\tOldReason: %s\n\tRefs: %s"),
-			*AssetData.ObjectPath.ToString(),
-			bMyResult ? 1 : 0,
-			bOldResult ? 1 : 0,
-			*MyReason.ToString(),
-			*OldReason.ToString(),
-			*FString::JoinBy(ReferencingDomains, TEXT(", "), [](TSharedPtr<FDomainData> Foo) { return Foo->UserFacingDomainName.ToString(); }));
-
-		if (bMyResult != bOldResult)
-		{
-			ensureAlways(false);
-		}
-	}
-#endif
-
 	if (OutOptionalFailureReason)
 	{
 		*OutOptionalFailureReason = MyReason;
diff --git a/Engine/Plugins/Editor/AssetReferenceRestrictions/Source/AssetReferenceRestrictions/Private/DomainAssetReferenceFilter.h b/Engine/Plugins/Editor/AssetReferenceRestrictions/Source/AssetReferenceRestrictions/Private/DomainAssetReferenceFilter.h
--- a/Engine/Plugins/Editor/AssetReferenceRestrictions/Source/AssetReferenceRestrictions/Private/DomainAssetReferenceFilter.h
+++ b/Engine/Plugins/Editor/AssetReferenceRestrictions/Source/AssetReferenceRestrictions/Private/DomainAssetReferenceFilter.h
@@ -8,9 +8,6 @@ struct FAssetData;
 struct FDomainDatabase;
 struct FDomainData;
 
-//@TODO: Remove the cross-checking
-#define UE_ASSET_DOMAIN_FILTERING_CROSS_CHECK_RESULTS 0
-
 class FDomainAssetReferenceFilter : public IAssetReferenceFilter
 {
 public:
@@ -39,9 +36,5 @@ private:
 	TArray<FAssetData> OriginalReferencingAssets;
 	TSet<TSharedPtr<FDomainData>> ReferencingDomains;
 
-#if UE_ASSET_DOMAIN_FILTERING_CROSS_CHECK_RESULTS
-	TSharedPtr<class FRefAssetReferenceFilter> CrossChecker;
-#endif
-
 	FText Failure_CouldNotDetermineDomain;
 };
diff --git a/Engine/Plugins/Editor/AssetReferenceRestrictions/Source/AssetReferenceRestrictions/Private/RefAssetReferenceFilter.cpp b/Engine/Plugins/Editor/AssetReferenceRestrictions/Source/AssetReferenceRestrictions/Private/RefAssetReferenceFilter.cpp
deleted file mode 100644
diff --git a/Engine/Plugins/Editor/AssetReferenceRestrictions/Source/AssetReferenceRestrictions/Private/RefAssetReferenceFilter.h b/Engine/Plugins/Editor/AssetReferenceRestrictions/Source/AssetReferenceRestrictions/Private/RefAssetReferenceFilter.h
deleted file mode 100644
