Asset referencing fixes
- Added a /Script domain to asset referencing so 'place actors' and similar circumstances can find their domain (It is not further split up by owning module (to match existing behavior in the other game and because there aren't packages per module in /Script))
- Prevented an ensure when an incomplete additional domain entry is created (new root path that hasn't been filled out yet)
- Changed the compiled-in log verbosity for LogAssetReferenceRestrictions to be Log as intended

diff --git a/Engine/Plugins/Editor/AssetReferenceRestrictions/Source/AssetReferenceRestrictions/Private/AssetReferencingDomains.cpp b/Engine/Plugins/Editor/AssetReferenceRestrictions/Source/AssetReferenceRestrictions/Private/AssetReferencingDomains.cpp
--- a/Engine/Plugins/Editor/AssetReferenceRestrictions/Source/AssetReferenceRestrictions/Private/AssetReferencingDomains.cpp
+++ b/Engine/Plugins/Editor/AssetReferenceRestrictions/Source/AssetReferenceRestrictions/Private/AssetReferencingDomains.cpp
@@ -127,6 +127,7 @@ void FDomainDatabase::RebuildFromScratch()
 	DomainNameMap.Reset();
 	PathMap.Reset();
 	EngineDomain.Reset();
+	ScriptDomain.Reset();
 	TempDomain.Reset();
 	GameDomain.Reset();
 	DomainsDefinedByPlugins.Reset();
@@ -137,6 +138,11 @@ void FDomainDatabase::RebuildFromScratch()
 	EngineDomain->UserFacingDomainName = LOCTEXT("EngineDomain", "EngineContent");
 	EngineDomain->DomainRootPaths.Add(TEXT("/Engine/"));
 
+	ScriptDomain = FindOrAddDomainByName(UAssetReferencingPolicySettings::ScriptDomainName);
+	ScriptDomain->UserFacingDomainName = LOCTEXT("ScriptDomain", "Script");
+	ScriptDomain->bCanBeSeenByEverything = true;
+	ScriptDomain->DomainRootPaths.Add(TEXT("/Script/"));
+
 	TempDomain = FindOrAddDomainByName(UAssetReferencingPolicySettings::TempDomainName);
 	TempDomain->UserFacingDomainName = LOCTEXT("TempDomain", "Temp");
 	TempDomain->bCanSeeEverything = true;
@@ -163,7 +169,7 @@ void FDomainDatabase::RebuildFromScratch()
 
 			for (const FDirectoryPath& PathRoot : PathRule.ContentRoots)
 			{
-				if (ensure(PathRoot.Path.StartsWith(TEXT("/"), ESearchCase::CaseSensitive)))
+				if (PathRoot.Path.StartsWith(TEXT("/"), ESearchCase::CaseSensitive))
 				{
 					Domain->DomainRootPaths.Add(PathRoot.Path / TEXT(""));
 				}
diff --git a/Engine/Plugins/Editor/AssetReferenceRestrictions/Source/AssetReferenceRestrictions/Private/AssetReferencingDomains.h b/Engine/Plugins/Editor/AssetReferenceRestrictions/Source/AssetReferenceRestrictions/Private/AssetReferencingDomains.h
--- a/Engine/Plugins/Editor/AssetReferenceRestrictions/Source/AssetReferenceRestrictions/Private/AssetReferencingDomains.h
+++ b/Engine/Plugins/Editor/AssetReferenceRestrictions/Source/AssetReferenceRestrictions/Private/AssetReferencingDomains.h
@@ -11,7 +11,7 @@
 #if UE_ASSET_DOMAIN_FILTERING_DEBUG_LOGGING
 DECLARE_LOG_CATEGORY_EXTERN(LogAssetReferenceRestrictions, Verbose, All);
 #else
-DECLARE_LOG_CATEGORY_EXTERN(LogAssetReferenceRestrictions, Log, Display);
+DECLARE_LOG_CATEGORY_EXTERN(LogAssetReferenceRestrictions, Log, Log);
 #endif
 
 class IPlugin;
@@ -86,6 +86,9 @@ private:
 	// The engine content domain
 	TSharedPtr<FDomainData> EngineDomain;
 
+	// The default script domain (for reflected types)
+	TSharedPtr<FDomainData> ScriptDomain;
+
 	// Used for various 'special' mount points like /Temp/, /Memory/, and /Extra/
 	// Not visible as a domain for other domains to see, and can see everything
 	TSharedPtr<FDomainData> TempDomain;
diff --git a/Engine/Plugins/Editor/AssetReferenceRestrictions/Source/AssetReferenceRestrictions/Private/AssetReferencingPolicySettings.cpp b/Engine/Plugins/Editor/AssetReferenceRestrictions/Source/AssetReferenceRestrictions/Private/AssetReferencingPolicySettings.cpp
--- a/Engine/Plugins/Editor/AssetReferenceRestrictions/Source/AssetReferenceRestrictions/Private/AssetReferencingPolicySettings.cpp
+++ b/Engine/Plugins/Editor/AssetReferenceRestrictions/Source/AssetReferenceRestrictions/Private/AssetReferencingPolicySettings.cpp
@@ -6,6 +6,7 @@
 #include "Editor.h"
 
 const FString UAssetReferencingPolicySettings::EngineDomainName(TEXT("EngineContent"));
+const FString UAssetReferencingPolicySettings::ScriptDomainName(TEXT("Script"));
 const FString UAssetReferencingPolicySettings::GameDomainName(TEXT("ProjectContent"));
 const FString UAssetReferencingPolicySettings::TempDomainName(TEXT("Temp"));
 
diff --git a/Engine/Plugins/Editor/AssetReferenceRestrictions/Source/AssetReferenceRestrictions/Private/AssetReferencingPolicySettings.h b/Engine/Plugins/Editor/AssetReferenceRestrictions/Source/AssetReferenceRestrictions/Private/AssetReferencingPolicySettings.h
--- a/Engine/Plugins/Editor/AssetReferenceRestrictions/Source/AssetReferenceRestrictions/Private/AssetReferencingPolicySettings.h
+++ b/Engine/Plugins/Editor/AssetReferenceRestrictions/Source/AssetReferenceRestrictions/Private/AssetReferencingPolicySettings.h
@@ -200,6 +200,7 @@ public:
 
 	// The names of the project, special system mount, and game domains
 	static const FString EngineDomainName;
+	static const FString ScriptDomainName;
 	static const FString GameDomainName;
 	static const FString TempDomainName;
 
