Fixed pickers in project settings (e.g., "Startup Map") not showing any project content when the asset referencing restrictions plugin is enabled

diff --git a/Engine/Plugins/Editor/AssetReferenceRestrictions/Source/AssetReferenceRestrictions/Private/AssetReferencingDomains.cpp b/Engine/Plugins/Editor/AssetReferenceRestrictions/Source/AssetReferenceRestrictions/Private/AssetReferencingDomains.cpp
--- a/Engine/Plugins/Editor/AssetReferenceRestrictions/Source/AssetReferenceRestrictions/Private/AssetReferencingDomains.cpp
+++ b/Engine/Plugins/Editor/AssetReferenceRestrictions/Source/AssetReferenceRestrictions/Private/AssetReferencingDomains.cpp
@@ -141,6 +141,7 @@ void FDomainDatabase::RebuildFromScratch()
 	ScriptDomain = FindOrAddDomainByName(UAssetReferencingPolicySettings::ScriptDomainName);
 	ScriptDomain->UserFacingDomainName = LOCTEXT("ScriptDomain", "Script");
 	ScriptDomain->bCanBeSeenByEverything = true;
+	ScriptDomain->bCanSeeEverything = true;
 	ScriptDomain->DomainRootPaths.Add(TEXT("/Script/"));
 
 	TempDomain = FindOrAddDomainByName(UAssetReferencingPolicySettings::TempDomainName);
