Excluding Transient package from external references in the validation checker.  Seems that compilation errors are bringing these into the AssetRegistry dependencies

diff --git a/Engine/Plugins/Editor/AssetReferenceRestrictions/Source/AssetReferenceRestrictions/Private/AssetValidator_AssetReferenceRestrictions.cpp b/Engine/Plugins/Editor/AssetReferenceRestrictions/Source/AssetReferenceRestrictions/Private/AssetValidator_AssetReferenceRestrictions.cpp
--- a/Engine/Plugins/Editor/AssetReferenceRestrictions/Source/AssetReferenceRestrictions/Private/AssetValidator_AssetReferenceRestrictions.cpp
+++ b/Engine/Plugins/Editor/AssetReferenceRestrictions/Source/AssetReferenceRestrictions/Private/AssetValidator_AssetReferenceRestrictions.cpp
@@ -45,6 +45,7 @@ EDataValidationResult UAssetValidator_AssetReferenceRestrictions::ValidateLoaded
 
 	FAssetRegistryModule& AssetRegistryModule = FModuleManager::LoadModuleChecked<FAssetRegistryModule>(TEXT("AssetRegistry"));
 	IAssetRegistry& AssetRegistry = AssetRegistryModule.Get();
+	const FName TransientName = GetTransientPackage()->GetFName();
 
 	// Check for missing soft or hard references to cinematic and developers content
 	FName PackageFName = InAsset->GetOutermost()->GetFName();
@@ -60,7 +61,10 @@ EDataValidationResult UAssetValidator_AssetReferenceRestrictions::ValidateLoaded
 			AssetRegistry.GetAssetsByPackageName(SoftDependency, DependencyAssets, true);
 			if (DependencyAssets.Num() == 0)
 			{
-				AssetFails(InAsset, FText::Format(LOCTEXT("IllegalReference_MissingSoftRef", "Soft references {0} which does not exist"), FText::FromString(SoftDependencyStr)), ValidationErrors);
+				if (SoftDependency != TransientName)
+				{
+					AssetFails(InAsset, FText::Format(LOCTEXT("IllegalReference_MissingSoftRef", "Soft references {0} which does not exist"), FText::FromString(SoftDependencyStr)), ValidationErrors);
+				}
 			}
 			else
 			{
