Remove temporary asset validation toggles from AncientGame development (bEnableAssetFiltering_TEMP & bEnableAssetValidator_TEMP)
#rnx

diff --git a/Engine/Plugins/Editor/AssetReferenceRestrictions/Source/AssetReferenceRestrictions/Private/AssetReferencingPolicySettings.h b/Engine/Plugins/Editor/AssetReferenceRestrictions/Source/AssetReferenceRestrictions/Private/AssetReferencingPolicySettings.h
--- a/Engine/Plugins/Editor/AssetReferenceRestrictions/Source/AssetReferenceRestrictions/Private/AssetReferencingPolicySettings.h
+++ b/Engine/Plugins/Editor/AssetReferenceRestrictions/Source/AssetReferenceRestrictions/Private/AssetReferencingPolicySettings.h
@@ -188,16 +188,6 @@ public:
 	UPROPERTY(config, EditAnywhere, Category="Project Content", meta=(TitleProperty=DomainName))
 	TArray<FARPDomainDefinitionByContentRoot> AdditionalDomains;
 
-	// This is a temporary flags for testing and should not be disabled under normal circumstances
-	// (it will be removed by 5.0)
-	UPROPERTY(config)
-	bool bEnableAssetFiltering_TEMP = true;
-
-	// This is a temporary flag for testing and should not be disabled under normal circumstances
-	// (it will be removed by 5.0)
-	UPROPERTY(config)
-	bool bEnableAssetValidator_TEMP = true;
-
 	// The names of the project, special system mount, and game domains
 	static const FString EngineDomainName;
 	static const FString ScriptDomainName;
diff --git a/Engine/Plugins/Editor/AssetReferenceRestrictions/Source/AssetReferenceRestrictions/Private/AssetReferencingPolicySubsystem.cpp b/Engine/Plugins/Editor/AssetReferenceRestrictions/Source/AssetReferenceRestrictions/Private/AssetReferencingPolicySubsystem.cpp
--- a/Engine/Plugins/Editor/AssetReferenceRestrictions/Source/AssetReferenceRestrictions/Private/AssetReferencingPolicySubsystem.cpp
+++ b/Engine/Plugins/Editor/AssetReferenceRestrictions/Source/AssetReferenceRestrictions/Private/AssetReferencingPolicySubsystem.cpp
@@ -10,7 +10,7 @@
 
 bool UAssetReferencingPolicySubsystem::ShouldCreateSubsystem(UObject* Outer) const
 {
-	return GetDefault<UAssetReferencingPolicySettings>()->bEnableAssetFiltering_TEMP;
+	return true;
 }
 
 void UAssetReferencingPolicySubsystem::Initialize(FSubsystemCollectionBase& Collection)
diff --git a/Engine/Plugins/Editor/AssetReferenceRestrictions/Source/AssetReferenceRestrictions/Private/AssetValidator_AssetReferenceRestrictions.cpp b/Engine/Plugins/Editor/AssetReferenceRestrictions/Source/AssetReferenceRestrictions/Private/AssetValidator_AssetReferenceRestrictions.cpp
--- a/Engine/Plugins/Editor/AssetReferenceRestrictions/Source/AssetReferenceRestrictions/Private/AssetValidator_AssetReferenceRestrictions.cpp
+++ b/Engine/Plugins/Editor/AssetReferenceRestrictions/Source/AssetReferenceRestrictions/Private/AssetValidator_AssetReferenceRestrictions.cpp
@@ -19,11 +19,6 @@ UAssetValidator_AssetReferenceRestrictions::UAssetValidator_AssetReferenceRestri
 
 bool UAssetValidator_AssetReferenceRestrictions::CanValidateAsset_Implementation(UObject* InAsset) const
 {
-	if (!GetDefault<UAssetReferencingPolicySettings>()->bEnableAssetValidator_TEMP)
-	{
-		return false;
-	}
-
 	if (InAsset)
 	{
 		const FAssetData AssetData(InAsset);
