Prevent a log error when adding new entries to CanReferenceTheseDomains

diff --git a/Engine/Plugins/Editor/AssetReferenceRestrictions/Source/AssetReferenceRestrictions/Private/AssetReferencingDomains.cpp b/Engine/Plugins/Editor/AssetReferenceRestrictions/Source/AssetReferenceRestrictions/Private/AssetReferencingDomains.cpp
--- a/Engine/Plugins/Editor/AssetReferenceRestrictions/Source/AssetReferenceRestrictions/Private/AssetReferencingDomains.cpp
+++ b/Engine/Plugins/Editor/AssetReferenceRestrictions/Source/AssetReferenceRestrictions/Private/AssetReferencingDomains.cpp
@@ -494,7 +494,10 @@ void FDomainDatabase::AddDomainVisibilityList(TSharedPtr<FDomainData> Domain, co
 {
 	for (const FString& OtherDomainName : VisibilityList)
 	{
-		Domain->DomainsVisibleFromHere.Add(FindOrAddDomainByName(OtherDomainName));
+		if (!OtherDomainName.IsEmpty())
+		{
+			Domain->DomainsVisibleFromHere.Add(FindOrAddDomainByName(OtherDomainName));
+		}
 	}
 }
 
