Update the tooltip on FARPDefaultProjectDomainRules::CanReferenceTheseDomains to reference the '/Game/' directory instead of calling it "GameContent". "GameContent" is not used anywhere else in the codebase for UE and can be confusing to people.

#rnx

diff --git a/Engine/Plugins/Editor/AssetReferenceRestrictions/Source/AssetReferenceRestrictions/Private/AssetReferencingPolicySettings.h b/Engine/Plugins/Editor/AssetReferenceRestrictions/Source/AssetReferenceRestrictions/Private/AssetReferencingPolicySettings.h
--- a/Engine/Plugins/Editor/AssetReferenceRestrictions/Source/AssetReferenceRestrictions/Private/AssetReferencingPolicySettings.h
+++ b/Engine/Plugins/Editor/AssetReferenceRestrictions/Source/AssetReferenceRestrictions/Private/AssetReferencingPolicySettings.h
@@ -108,8 +108,10 @@ struct FARPDefaultProjectDomainRules
 {
 	GENERATED_BODY()
 
-	// The list of additional domains always visible from this domain
-	// (EngineContent and GameContent are always visible)
+	/**
+	 * The list of additional domains always visible from this domain.
+	 * EngineContent and Project Content in the '/Game/*' directory are always visible
+	 */
 	UPROPERTY(EditAnywhere, Category = Settings, meta=(GetOptions="GetListOfDomains_NoEngineOrGame"))
 	TArray<FString> CanReferenceTheseDomains;
 };
