Fix linux build error because of a "/*" in a comment

#rnx

diff --git a/Engine/Plugins/Editor/AssetReferenceRestrictions/Source/AssetReferenceRestrictions/Private/AssetReferencingPolicySettings.h b/Engine/Plugins/Editor/AssetReferenceRestrictions/Source/AssetReferenceRestrictions/Private/AssetReferencingPolicySettings.h
--- a/Engine/Plugins/Editor/AssetReferenceRestrictions/Source/AssetReferenceRestrictions/Private/AssetReferencingPolicySettings.h
+++ b/Engine/Plugins/Editor/AssetReferenceRestrictions/Source/AssetReferenceRestrictions/Private/AssetReferencingPolicySettings.h
@@ -110,7 +110,7 @@ struct FARPDefaultProjectDomainRules
 
 	/**
 	 * The list of additional domains always visible from this domain.
-	 * EngineContent and Project Content in the '/Game/*' directory are always visible
+	 * EngineContent and Project Content in the '/Game/' directory are always visible
 	 */
 	UPROPERTY(EditAnywhere, Category = Settings, meta=(GetOptions="GetListOfDomains_NoEngineOrGame"))
 	TArray<FString> CanReferenceTheseDomains;
