Fix ExternalActors from the Developers folder seen as error by the AssetReferenceRestrictions validations plugin.

#rnx

diff --git a/Engine/Plugins/Editor/AssetReferenceRestrictions/Source/AssetReferenceRestrictions/Private/AssetReferencingDomains.cpp b/Engine/Plugins/Editor/AssetReferenceRestrictions/Source/AssetReferenceRestrictions/Private/AssetReferencingDomains.cpp
--- a/Engine/Plugins/Editor/AssetReferenceRestrictions/Source/AssetReferenceRestrictions/Private/AssetReferencingDomains.cpp
+++ b/Engine/Plugins/Editor/AssetReferenceRestrictions/Source/AssetReferenceRestrictions/Private/AssetReferencingDomains.cpp
@@ -6,6 +6,8 @@
 #include "DomainAssetReferenceFilter.h"
 #include "Editor.h"
 #include "Settings/ProjectPackagingSettings.h"
+#include "Misc/PackagePath.h"
+#include "Misc/StringBuilder.h"
 
 DEFINE_LOG_CATEGORY(LogAssetReferenceRestrictions);
 
@@ -451,7 +453,30 @@ TSharedPtr<FDomainData> FDomainDatabase::FindDomainFromAssetData(const FAssetDat
 	
 	if(!DomainResult.IsValid())
 	{
-		DomainResult = PathMap->FindDomainFromPath(AssetData.PackagePath.ToString());
+		FString PackagePath = AssetData.PackagePath.ToString();
+
+		// Treat external objects to be part of the same domain as their owning objects, so we don't treat them as being part of /Game/... just 
+		// because they are in /Game/ExternalObjects/xyz/... but as if they were in /Game/xyz/... instead, for the context of asset restriction 
+		// validations.
+		TArray<FString> SplitPath;
+		if (PackagePath.ParseIntoArray(SplitPath, TEXT("/")) > 1)
+		{
+			if ((SplitPath[1] == FPackagePath::GetExternalObjectsFolderName()) || (SplitPath[1] == FPackagePath::GetExternalActorsFolderName()))
+			{
+				SplitPath.RemoveAt(1);
+
+				// Reconstruct path without any external objects paths
+				TStringBuilder<1024> NewPackagePath;
+				for (const FString& Path : SplitPath)
+				{
+					NewPackagePath += TEXT("/");
+					NewPackagePath += Path;
+				}
+				PackagePath = NewPackagePath.ToString();
+			}
+		}
+
+		DomainResult = PathMap->FindDomainFromPath(PackagePath);
 	}
 
 #if UE_ASSET_DOMAIN_FILTERING_DEBUG_LOGGING
