Fix for crash when validating data.

#test Run tools -> Validate Assets in editor

diff --git a/Engine/Plugins/Editor/AssetReferenceRestrictions/Source/AssetReferenceRestrictions/Private/AssetValidator_AssetReferenceRestrictions.cpp b/Engine/Plugins/Editor/AssetReferenceRestrictions/Source/AssetReferenceRestrictions/Private/AssetValidator_AssetReferenceRestrictions.cpp
--- a/Engine/Plugins/Editor/AssetReferenceRestrictions/Source/AssetReferenceRestrictions/Private/AssetValidator_AssetReferenceRestrictions.cpp
+++ b/Engine/Plugins/Editor/AssetReferenceRestrictions/Source/AssetReferenceRestrictions/Private/AssetValidator_AssetReferenceRestrictions.cpp
@@ -24,7 +24,7 @@ bool UAssetValidator_AssetReferenceRestrictions::CanValidateAsset_Implementation
 		const FAssetData AssetData(InAsset);
 		TSharedPtr<FDomainData> DomainData = GEditor->GetEditorSubsystem<UAssetReferencingPolicySubsystem>()->GetDomainDB()->FindDomainFromAssetData(AssetData);
 
-		const bool bIsInUnrestrictedFolder = DomainData->IsValid() && DomainData->bCanSeeEverything;
+		const bool bIsInUnrestrictedFolder = DomainData && DomainData->IsValid() && DomainData->bCanSeeEverything;
 		if (!bIsInUnrestrictedFolder)
 		{
 			return true;
