Updated AssetReferenceRestrictions not to set "OptimizeCode = CodeOptimization.Never" so it shares the same PCH as everything else. This changes removes the need to create a 1.84 GB PCH specifically for this module and improves overall compile times.

Before:
Total CPU Time: 33.90625 s
Total time in Parallel executor: 31.99 seconds
After:
Total CPU Time: 7.859375 s
Total time in Parallel executor: 4.63 seconds

diff --git a/Engine/Plugins/Editor/AssetReferenceRestrictions/Source/AssetReferenceRestrictions/AssetReferenceRestrictions.Build.cs b/Engine/Plugins/Editor/AssetReferenceRestrictions/Source/AssetReferenceRestrictions/AssetReferenceRestrictions.Build.cs
--- a/Engine/Plugins/Editor/AssetReferenceRestrictions/Source/AssetReferenceRestrictions/AssetReferenceRestrictions.Build.cs
+++ b/Engine/Plugins/Editor/AssetReferenceRestrictions/Source/AssetReferenceRestrictions/AssetReferenceRestrictions.Build.cs
@@ -6,8 +6,6 @@ namespace UnrealBuildTool.Rules
 	{
         public AssetReferenceRestrictions(ReadOnlyTargetRules Target) : base(Target)
 		{
-			OptimizeCode = CodeOptimization.Never;
-
 			PublicDependencyModuleNames.AddRange(
 				new string[]
                 {
