Exposed AssetReferencingPolicy functionality for use in other modules

diff --git a/Engine/Plugins/Editor/AssetReferenceRestrictions/Source/AssetReferenceRestrictions/Private/AssetReferencingDomains.cpp b/Engine/Plugins/Editor/AssetReferenceRestrictions/Source/AssetReferenceRestrictions/Private/AssetReferencingDomains.cpp
--- a/Engine/Plugins/Editor/AssetReferenceRestrictions/Source/AssetReferenceRestrictions/Private/AssetReferencingDomains.cpp
+++ b/Engine/Plugins/Editor/AssetReferenceRestrictions/Source/AssetReferenceRestrictions/Private/AssetReferencingDomains.cpp
@@ -262,6 +262,8 @@ void FDomainDatabase::RebuildFromScratch()
 
 	// Update any existing filters that have cached now stale domains
 	FDomainAssetReferenceFilter::UpdateAllFilters();
+
+	OnPostDatabaseUpdated.Broadcast();
 }
 
 void FDomainDatabase::Init()
diff --git a/Engine/Plugins/Editor/AssetReferenceRestrictions/Source/AssetReferenceRestrictions/Private/AssetReferencingDomains.h b/Engine/Plugins/Editor/AssetReferenceRestrictions/Source/AssetReferenceRestrictions/Private/AssetReferencingDomains.h
deleted file mode 100644
diff --git a/Engine/Plugins/Editor/AssetReferenceRestrictions/Source/AssetReferenceRestrictions/Private/AssetReferencingPolicySettings.h b/Engine/Plugins/Editor/AssetReferenceRestrictions/Source/AssetReferenceRestrictions/Private/AssetReferencingPolicySettings.h
deleted file mode 100644
diff --git a/Engine/Plugins/Editor/AssetReferenceRestrictions/Source/AssetReferenceRestrictions/Private/AssetReferencingPolicySubsystem.h b/Engine/Plugins/Editor/AssetReferenceRestrictions/Source/AssetReferenceRestrictions/Private/AssetReferencingPolicySubsystem.h
deleted file mode 100644
diff --git a/Engine/Plugins/Editor/AssetReferenceRestrictions/Source/AssetReferenceRestrictions/Private/DomainAssetReferenceFilter.h b/Engine/Plugins/Editor/AssetReferenceRestrictions/Source/AssetReferenceRestrictions/Private/DomainAssetReferenceFilter.h
deleted file mode 100644
diff --git a/Engine/Plugins/Editor/AssetReferenceRestrictions/Source/AssetReferenceRestrictions/Public/AssetReferencingDomains.h b/Engine/Plugins/Editor/AssetReferenceRestrictions/Source/AssetReferenceRestrictions/Public/AssetReferencingDomains.h
new file mode 100644
--- /dev/null
+++ b/Engine/Plugins/Editor/AssetReferenceRestrictions/Source/AssetReferenceRestrictions/Public/AssetReferencingDomains.h
@@ -0,0 +1,118 @@
+// Copyright Epic Games, Inc. All Rights Reserved.
+
+#pragma once
+
+#include "Internationalization/Text.h"
+
+//@TODO: Debug printing
+#define UE_ASSET_DOMAIN_FILTERING_DEBUG_LOGGING 0
+
+#if UE_ASSET_DOMAIN_FILTERING_DEBUG_LOGGING
+DECLARE_LOG_CATEGORY_EXTERN(LogAssetReferenceRestrictions, Verbose, All);
+#else
+DECLARE_LOG_CATEGORY_EXTERN(LogAssetReferenceRestrictions, Log, Log);
+#endif
+
+class IPlugin;
+struct FDomainPathNode;
+struct FAssetData;
+
+struct FDomainData final : public TSharedFromThis<FDomainData>
+{
+	FText UserFacingDomainName;
+	FText ErrorMessageIfUsedElsewhere;
+
+	// The list of root paths, always of the format /Mount/ or /Mount/Path/To/ with both leading and trailing /
+	TArray<FString> DomainRootPaths;
+
+	// A list of specific packages that are part of this domain
+	TArray<FName> SpecificAssetPackages;
+
+	// The domains that are visible from here (if bCanSeeEverything is true, then literally everything is visible from here)
+	TSet<TSharedPtr<FDomainData>> DomainsVisibleFromHere;
+
+	// Can we see everything?
+	bool bCanSeeEverything = false;
+
+	// Can we be seen by everything?
+	bool bCanBeSeenByEverything = false;
+
+	bool IsValid() const
+	{
+		return DomainRootPaths.Num() > 0 || SpecificAssetPackages.Num() > 0;
+	}
+
+	void Reset()
+	{
+		UserFacingDomainName = FText::GetEmpty();
+		DomainRootPaths.Reset();
+		DomainsVisibleFromHere.Reset();
+		bCanSeeEverything = false;
+	}
+};
+
+struct ASSETREFERENCERESTRICTIONS_API FDomainDatabase final
+{
+	FDomainDatabase();
+	~FDomainDatabase();
+
+	void Init();
+
+	void MarkDirty();
+	void UpdateIfNecessary();
+
+	// Delegate that's called whenever the database has been updated to allow people to respond to the change
+	TMulticastDelegate<void()>& PostDatabaseUpdated() { return OnPostDatabaseUpdated; };
+
+	void ValidateAllDomains();
+	void DebugPrintAllDomains();
+
+	void OnPluginCreatedOrMounted(IPlugin& NewPlugin);
+
+	TSharedPtr<FDomainData> FindDomainFromAssetData(const FAssetData& AssetData) const;
+
+	TTuple<bool, FText> CanDomainsSeeEachOther(TSharedPtr<FDomainData> Referencee, TSharedPtr<FDomainData> Referencer) const;
+
+	const TArray<FString>& GetDomainsDefinedByPlugins() const { return DomainsDefinedByPlugins; }
+
+	TSharedPtr<FDomainData> FindOrAddDomainByName(const FString& Name);
+
+private:
+	void RebuildFromScratch();
+	
+	void BuildDomainFromPlugin(TSharedRef<IPlugin> Plugin);
+
+	void AddDomainVisibilityList(TSharedPtr<FDomainData> Domain, const TArray<FString>& VisibilityList);
+private:
+	// Map from domain name to 
+	TMap<FString, TSharedPtr<FDomainData>> DomainNameMap;
+
+	// Map from path to domain
+	TSharedPtr<FDomainPathNode> PathMap;
+
+	// Specific packages that are forced into certain domains
+	TMap<FName, TSharedPtr<FDomainData>> SpecificAssetPackageDomains;
+
+	// The engine content domain
+	TSharedPtr<FDomainData> EngineDomain;
+
+	// The default script domain (for reflected types)
+	TSharedPtr<FDomainData> ScriptDomain;
+
+	// Used for various 'special' mount points like /Temp/, /Memory/, and /Extra/
+	// Not visible as a domain for other domains to see, and can see everything
+	TSharedPtr<FDomainData> TempDomain;
+
+	// The game content domain
+	TSharedPtr<FDomainData> GameDomain;
+
+	// The never cooked content domain
+	TSharedPtr<FDomainData> NeverCookDomain;
+
+	// List of domains that came from plugins (used for domain pickers in the settings)
+	TArray<FString> DomainsDefinedByPlugins;
+
+	bool bDatabaseOutOfDate = false;
+
+	TMulticastDelegate<void()> OnPostDatabaseUpdated;
+};
diff --git a/Engine/Plugins/Editor/AssetReferenceRestrictions/Source/AssetReferenceRestrictions/Public/AssetReferencingPolicySettings.h b/Engine/Plugins/Editor/AssetReferenceRestrictions/Source/AssetReferenceRestrictions/Public/AssetReferencingPolicySettings.h
new file mode 100644
--- /dev/null
+++ b/Engine/Plugins/Editor/AssetReferenceRestrictions/Source/AssetReferenceRestrictions/Public/AssetReferencingPolicySettings.h
@@ -0,0 +1,242 @@
+// Copyright Epic Games, Inc. All Rights Reserved.
+
+#pragma once
+
+#include "Engine/DeveloperSettings.h"
+
+#include "AssetReferencingPolicySettings.generated.h"
+
+struct FDirectoryPath;
+struct FPropertyChangedEvent;
+
+USTRUCT()
+struct FARPDefaultPluginDomainRules
+{
+	GENERATED_BODY()
+
+	// The list of additional domains always visible from a plugin
+	// (EngineContent is always visible, as is content from other plugins that are explicitly referenced)
+	UPROPERTY(EditAnywhere, Category=Settings, meta=(GetOptions="GetListOfDomains_NoPluginsOrEngine"))
+	TArray<FString> CanReferenceTheseDomains;
+
+	// Can content in the ProjectContent domain access plugin content automatically (for plugins that don't match a specific rule)?
+	UPROPERTY(config, EditAnywhere, Category=Settings)
+	bool bCanProjectAccessThesePlugins = true;
+
+	// Can content in other domains access plugin content automatically (for plugins that don't match a specific rule)?
+	// Note: This rule may be deprecated in the future!
+	UPROPERTY(config, EditAnywhere, Category=Settings)
+	bool bCanBeSeenByOtherDomainsWithoutDependency = false;
+};
+
+UENUM()
+enum class EARPPluginMatchMode : uint8
+{
+	MatchByCategory,
+	MatchByPathPrefix
+};
+
+USTRUCT()
+struct FARPDomainDefinitionForMatchingPlugins
+{
+	GENERATED_BODY()
+
+	// The display name of this rule (used in error message when attempting to reference content incorrectly)
+	// The token {0} will be replaced with the plugin name
+	UPROPERTY(EditAnywhere, Category=Settings)
+	FText DisplayName;
+
+	// The error message if something that is not allowed to attempts to reference content from this domain
+	UPROPERTY(EditAnywhere, Category=Settings)
+	FText ErrorMessageIfUsedElsewhere;
+
+	// Type of matching for this rule
+	UPROPERTY(EditAnywhere, Category=Settings)
+	EARPPluginMatchMode MatchRule = EARPPluginMatchMode::MatchByPathPrefix;
+
+	// If set, a plugin with the same rooted directory path will match this rule
+	// (use "/FirstFolder/SecondFolder/" to match a plugin like $YourProjectDir/Plugins/FirstFolder/SecondFolder/MyCoolPlugin/MyCoolPlugin.uplugin)
+	UPROPERTY(EditAnywhere, Category=Settings, meta=(EditCondition="MatchRule==EARPPluginMatchMode::MatchByPathPrefix", EditConditionHides))
+	FString PluginPathPrefix;
+
+	// If set, a plugin with a matching Category will match this rule
+	UPROPERTY(EditAnywhere, Category=Settings, meta=(EditCondition="MatchRule==EARPPluginMatchMode::MatchByCategory", EditConditionHides))
+	FString PluginCategoryPrefix;
+
+	// The list of additional domains always visible from a plugin
+	// (EngineContent is always visible, as is content from other plugins that are explicitly referenced)
+	UPROPERTY(EditAnywhere, Category=Settings, meta=(GetOptions="GetListOfDomains_NoPluginsOrEngine"))
+	TArray<FString> CanReferenceTheseDomains;
+
+	bool IsValid() const
+	{
+		if (DisplayName.IsEmpty()) { return false; }
+
+		if (MatchRule == EARPPluginMatchMode::MatchByCategory)
+		{
+			if (PluginCategoryPrefix.IsEmpty()) { return false; }
+		}
+		else
+		{
+			if (PluginPathPrefix.IsEmpty()) { return false; }
+		}
+
+		return true;
+	}
+};
+
+USTRUCT()
+struct FARPDomainSettingsForPlugins
+{
+	GENERATED_BODY()
+
+	// The default rule if a more specific plugin rule doesn't apply
+	UPROPERTY(EditAnywhere, Category=Settings)
+	FARPDefaultPluginDomainRules DefaultRule;
+
+	// Discovered plugins will be matched against these templates
+	// Priority rules (a path match is preferred to a category match, and within each the longest match wins):
+	//   Highest: The most specific path match
+	//            Any path match
+	//            The most specific category match
+	//            Any category match
+	UPROPERTY(EditAnywhere, Category=Settings, meta=(TitleProperty="DisplayName"))
+	TArray<FARPDomainDefinitionForMatchingPlugins> AdditionalRules;
+};
+
+
+USTRUCT()
+struct FARPDefaultProjectDomainRules
+{
+	GENERATED_BODY()
+
+	/**
+	 * The list of additional domains always visible from this domain.
+	 * EngineContent and Project Content in the '/Game/' directory are always visible
+	 */
+	UPROPERTY(EditAnywhere, Category = Settings, meta=(GetOptions="GetListOfDomains_NoEngineOrGame"))
+	TArray<FString> CanReferenceTheseDomains;
+};
+
+
+UENUM()
+enum class EARPDomainAllowedToReferenceMode : uint8
+{
+	AdditionalDomains,
+	AllDomains UMETA(DisplayName = "All Domains (DANGER)")
+};
+
+
+USTRUCT()
+struct FARPDomainDefinitionByContentRoot
+{
+	GENERATED_BODY()
+
+	// The name of this domain
+	UPROPERTY(EditAnywhere, Category=Settings)
+	FString DomainName;
+
+	// The display name of this domain (used in error message when attempting to reference content incorrectly)
+	UPROPERTY(EditAnywhere, Category=Settings)
+	FText DomainDisplayName;
+
+	// The error message if something that is not allowed to attempts to reference content from this domain
+	UPROPERTY(EditAnywhere, Category=Settings)
+	FText ErrorMessageIfUsedElsewhere;
+
+	// The list of content root paths considered to be part of this domain
+	UPROPERTY(EditAnywhere, Category=Settings, meta=(ContentDir))
+	TArray<FDirectoryPath> ContentRoots;
+
+	// A list of specific assets considered to be part of this domain
+	UPROPERTY(EditAnywhere, Category=Settings, meta=(LongPackageName))
+	TArray<FName> SpecificAssets;
+
+	// What content is this domain allowed to access?
+	UPROPERTY(EditAnywhere, Category = Settings)
+	EARPDomainAllowedToReferenceMode ReferenceMode = EARPDomainAllowedToReferenceMode::AdditionalDomains;
+
+	// The list of additional domains always visible from this domain
+	// (EngineContent is always visible)
+	UPROPERTY(EditAnywhere, Category = Settings, meta=(GetOptions="GetListOfDomains_NoEngine", EditCondition="ReferenceMode==EARPDomainAllowedToReferenceMode::AdditionalDomains", EditConditionHides))
+	TArray<FString> CanReferenceTheseDomains;
+
+	bool IsValid() const
+	{
+		if (DomainName.IsEmpty()) { return false; }
+
+		if (ContentRoots.Num() == 0 && SpecificAssets.Num() == 0) { return false; }
+
+		if (DomainDisplayName.IsEmpty()) { return false; }
+
+		return true;
+	}
+};
+
+/** Settings for the Asset Referencing Policy, these settings are used to determine which plugins and game folders can reference content from each other */
+UCLASS(config=Game, defaultconfig, meta=(DisplayName="Asset Referencing Policy"))
+class ASSETREFERENCERESTRICTIONS_API UAssetReferencingPolicySettings : public UDeveloperSettings
+{
+	GENERATED_BODY()
+
+public:
+	// Settings/rules for engine plugins
+	UPROPERTY(config, EditAnywhere, Category="Engine Plugins")
+	FARPDomainSettingsForPlugins EnginePlugins;
+
+	// Settings/rules for project plugins
+	UPROPERTY(config, EditAnywhere, Category="Project Plugins")
+	FARPDomainSettingsForPlugins ProjectPlugins;
+
+	// The default rules for project content (if a more specific rule doesn't apply)
+	UPROPERTY(config, EditAnywhere, Category="Project Content")
+	FARPDefaultProjectDomainRules DefaultProjectContentRule;
+
+	// List of additional domains to carve out from the project content folder
+	UPROPERTY(config, EditAnywhere, Category="Project Content", meta=(TitleProperty=DomainName))
+	TArray<FARPDomainDefinitionByContentRoot> AdditionalDomains;
+
+	// The names of the project, special system mount, and game domains
+	static const FString EngineDomainName;
+	static const FString ScriptDomainName;
+	static const FString GameDomainName;
+	static const FString TempDomainName;
+	static const FString NeverCookDomainName;
+
+public:
+	//~UObject interface
+#if WITH_EDITOR
+	virtual void PostEditChangeProperty(FPropertyChangedEvent& PropertyChangedEvent) override;
+#endif
+	//~End of UObject interface
+
+	UAssetReferencingPolicySettings(const FObjectInitializer& ObjectInitializer);
+
+private:
+	TArray<FString> GetListOfDomains(bool bAllowPlugins, bool bAllowEngine, bool bAllowGame) const;
+
+	UFUNCTION()
+	TArray<FString> GetListOfDomains_All() const
+	{
+		return GetListOfDomains(true, true, true);
+	}
+
+	UFUNCTION()
+	TArray<FString> GetListOfDomains_NoEngine() const
+	{
+		return GetListOfDomains(true, false, true);
+	}
+
+	UFUNCTION()
+	TArray<FString> GetListOfDomains_NoEngineOrGame() const
+	{
+		return GetListOfDomains(true, false, false);
+	}
+
+	UFUNCTION()
+	TArray<FString> GetListOfDomains_NoPluginsOrEngine() const
+	{
+		return GetListOfDomains(false, false, true);
+	}
+
+};
diff --git a/Engine/Plugins/Editor/AssetReferenceRestrictions/Source/AssetReferenceRestrictions/Public/AssetReferencingPolicySubsystem.h b/Engine/Plugins/Editor/AssetReferenceRestrictions/Source/AssetReferenceRestrictions/Public/AssetReferencingPolicySubsystem.h
new file mode 100644
--- /dev/null
+++ b/Engine/Plugins/Editor/AssetReferenceRestrictions/Source/AssetReferenceRestrictions/Public/AssetReferencingPolicySubsystem.h
@@ -0,0 +1,31 @@
+// Copyright Epic Games, Inc. All Rights Reserved.
+
+#pragma once
+
+#include "EditorSubsystem.h"
+
+#include "AssetReferencingPolicySubsystem.generated.h"
+
+struct FAssetReferenceFilterContext;
+class IAssetReferenceFilter;
+struct FDomainDatabase;
+
+/** Subsystem to register the domain-based asset referencing policy restrictions with the editor */
+UCLASS()
+class ASSETREFERENCERESTRICTIONS_API UAssetReferencingPolicySubsystem : public UEditorSubsystem
+{
+	GENERATED_BODY()
+
+public:
+	//~UEditorSubsystem interface
+	virtual bool ShouldCreateSubsystem(UObject* Outer) const override;
+	virtual void Initialize(FSubsystemCollectionBase& Collection) override;
+	virtual void Deinitialize() override;
+	//~End of UEditorSubsystem interface
+
+	TSharedPtr<FDomainDatabase> GetDomainDB() const;
+private:
+	TSharedPtr<IAssetReferenceFilter> HandleMakeAssetReferenceFilter(const FAssetReferenceFilterContext& Context);
+
+	TSharedPtr<FDomainDatabase> DomainDB;
+};
diff --git a/Engine/Plugins/Editor/AssetReferenceRestrictions/Source/AssetReferenceRestrictions/Public/DomainAssetReferenceFilter.h b/Engine/Plugins/Editor/AssetReferenceRestrictions/Source/AssetReferenceRestrictions/Public/DomainAssetReferenceFilter.h
new file mode 100644
--- /dev/null
+++ b/Engine/Plugins/Editor/AssetReferenceRestrictions/Source/AssetReferenceRestrictions/Public/DomainAssetReferenceFilter.h
@@ -0,0 +1,43 @@
+// Copyright Epic Games, Inc. All Rights Reserved.
+
+#pragma once
+
+#include "AssetRegistry/AssetData.h"
+#include "Editor/AssetReferenceFilter.h"
+
+struct FAssetReferenceFilterContext;
+
+struct FAssetData;
+struct FDomainDatabase;
+struct FDomainData;
+
+class ASSETREFERENCERESTRICTIONS_API FDomainAssetReferenceFilter : public IAssetReferenceFilter
+{
+public:
+	FDomainAssetReferenceFilter(const FAssetReferenceFilterContext& Context, TSharedPtr<FDomainDatabase> InDomainDB);
+	~FDomainAssetReferenceFilter();
+
+	//~IAssetReferenceFilter interface
+	virtual bool PassesFilter(const FAssetData& AssetData, FText* OutOptionalFailureReason = nullptr) const override;
+	//~End of IAssetReferenceFilter
+
+	// Update any cached information for all filters
+	static void UpdateAllFilters();
+private:
+	bool PassesFilterImpl(const FAssetData& AssetData, FText& OutOptionalFailureReason) const;
+
+	void DetermineReferencingDomain();
+
+	/** Heuristic to find actual assets from preview assets (i.e., the material editor's preview material) */
+	void TryGetAssociatedAssetsFromPossiblyPreviewObject(UObject* PossiblyPreviewObject, TArray<FAssetData>& InOutAssetsToConsider) const;
+
+private:
+	static TArray<FDomainAssetReferenceFilter*> FilterInstances;
+
+	TSharedPtr<FDomainDatabase> DomainDB;
+
+	TArray<FAssetData> OriginalReferencingAssets;
+	TSet<TSharedPtr<FDomainData>> ReferencingDomains;
+
+	FText Failure_CouldNotDetermineDomain;
+};
