Change FDomainDatabase::ValidateAllDomains to log a Warning instead of an Error. This change is temporary to unblock failing builds while a proper fix is being reviewed (see UE-184482).

diff --git a/Engine/Plugins/Editor/AssetReferenceRestrictions/Source/AssetReferenceRestrictions/Private/AssetReferencingDomains.cpp b/Engine/Plugins/Editor/AssetReferenceRestrictions/Source/AssetReferenceRestrictions/Private/AssetReferencingDomains.cpp
--- a/Engine/Plugins/Editor/AssetReferenceRestrictions/Source/AssetReferenceRestrictions/Private/AssetReferencingDomains.cpp
+++ b/Engine/Plugins/Editor/AssetReferenceRestrictions/Source/AssetReferenceRestrictions/Private/AssetReferencingDomains.cpp
@@ -431,7 +431,7 @@ void FDomainDatabase::ValidateAllDomains()
 				}
 			}
 
-			UE_LOG(LogAssetReferenceRestrictions, Error, TEXT("Asset domain %s was referenced by [%s] but wasn't found (perhaps it is a plugin that was disabled?)"), *KVP.Key, *FString::Join(DomainsThatReferenceMe, TEXT(", ")));
+			UE_LOG(LogAssetReferenceRestrictions, Warning, TEXT("Asset domain %s was referenced by [%s] but wasn't found (perhaps it is a plugin that was disabled?)"), *KVP.Key, *FString::Join(DomainsThatReferenceMe, TEXT(", ")));
 		}
 	}
 }
