Undo temp hack: Make AssetReferenceRestrictions log Errors again

diff --git a/Engine/Plugins/Editor/AssetReferenceRestrictions/Source/AssetReferenceRestrictions/Private/AssetReferencingDomains.cpp b/Engine/Plugins/Editor/AssetReferenceRestrictions/Source/AssetReferenceRestrictions/Private/AssetReferencingDomains.cpp
--- a/Engine/Plugins/Editor/AssetReferenceRestrictions/Source/AssetReferenceRestrictions/Private/AssetReferencingDomains.cpp
+++ b/Engine/Plugins/Editor/AssetReferenceRestrictions/Source/AssetReferenceRestrictions/Private/AssetReferencingDomains.cpp
@@ -431,7 +431,7 @@ void FDomainDatabase::ValidateAllDomains()
 				}
 			}
 
-			UE_LOG(LogAssetReferenceRestrictions, Warning, TEXT("Asset domain %s was referenced by [%s] but wasn't found (perhaps it is a plugin that was disabled?)"), *KVP.Key, *FString::Join(DomainsThatReferenceMe, TEXT(", ")));
+			UE_LOG(LogAssetReferenceRestrictions, Error, TEXT("Asset domain %s was referenced by [%s] but wasn't found (perhaps it is a plugin that was disabled?)"), *KVP.Key, *FString::Join(DomainsThatReferenceMe, TEXT(", ")));
 		}
 	}
 }
