Add bIgnoreEditorOnlyReferences to AssetReferencingPolicySettings, defaulting to true.
When this option is enabled, the reference checker will ignore editor-only references to things like editor utility components

diff --git a/Engine/Plugins/Editor/AssetReferenceRestrictions/Source/AssetReferenceRestrictions/Private/AssetValidator_AssetReferenceRestrictions.cpp b/Engine/Plugins/Editor/AssetReferenceRestrictions/Source/AssetReferenceRestrictions/Private/AssetValidator_AssetReferenceRestrictions.cpp
--- a/Engine/Plugins/Editor/AssetReferenceRestrictions/Source/AssetReferenceRestrictions/Private/AssetValidator_AssetReferenceRestrictions.cpp
+++ b/Engine/Plugins/Editor/AssetReferenceRestrictions/Source/AssetReferenceRestrictions/Private/AssetValidator_AssetReferenceRestrictions.cpp
@@ -4,6 +4,7 @@
 #include "AssetRegistry/AssetRegistryModule.h"
 #include "Editor.h"
 #include "AssetReferencingPolicySubsystem.h"
+#include "AssetReferencingPolicySettings.h"
 #include "AssetReferencingDomains.h"
 #include "Editor/AssetReferenceFilter.h"
 #include "Misc/PackageName.h"
@@ -47,7 +48,17 @@ EDataValidationResult UAssetValidator_AssetReferenceRestrictions::ValidateLoaded
 	FName PackageFName = InAsset->GetOutermost()->GetFName();
 	TArray<FName> SoftDependencies;
 	TArray<FAssetData> AllDependencyAssets;
-	AssetRegistry.GetDependencies(PackageFName, SoftDependencies, UE::AssetRegistry::EDependencyCategory::Package, UE::AssetRegistry::EDependencyQuery::Soft);
+
+	const UAssetReferencingPolicySettings* Settings = GetDefault<UAssetReferencingPolicySettings>();
+	UE::AssetRegistry::EDependencyQuery QueryFlags = UE::AssetRegistry::EDependencyQuery::NoRequirements;
+
+	if (Settings->bIgnoreEditorOnlyReferences)
+	{
+		// If the rules allow ignoring editor-only referencers, restrict the dependencies to game references
+		QueryFlags = UE::AssetRegistry::EDependencyQuery::Game;
+	}
+
+	AssetRegistry.GetDependencies(PackageFName, SoftDependencies, UE::AssetRegistry::EDependencyCategory::Package, UE::AssetRegistry::EDependencyQuery::Soft | QueryFlags);
 	for (FName SoftDependency : SoftDependencies)
 	{
 		const FString SoftDependencyStr = SoftDependency.ToString();
@@ -71,7 +82,7 @@ EDataValidationResult UAssetValidator_AssetReferenceRestrictions::ValidateLoaded
 
 	// Now check hard references to cinematic and developers content
 	TArray<FName> HardDependencies;
-	AssetRegistry.GetDependencies(PackageFName, HardDependencies, UE::AssetRegistry::EDependencyCategory::Package, UE::AssetRegistry::EDependencyQuery::Hard);
+	AssetRegistry.GetDependencies(PackageFName, HardDependencies, UE::AssetRegistry::EDependencyCategory::Package, UE::AssetRegistry::EDependencyQuery::Hard | QueryFlags);
 	for (FName HardDependency : HardDependencies)
 	{
 		//@TODO: Probably not needed anymore?
diff --git a/Engine/Plugins/Editor/AssetReferenceRestrictions/Source/AssetReferenceRestrictions/Public/AssetReferencingPolicySettings.h b/Engine/Plugins/Editor/AssetReferenceRestrictions/Source/AssetReferenceRestrictions/Public/AssetReferencingPolicySettings.h
--- a/Engine/Plugins/Editor/AssetReferenceRestrictions/Source/AssetReferenceRestrictions/Public/AssetReferencingPolicySettings.h
+++ b/Engine/Plugins/Editor/AssetReferenceRestrictions/Source/AssetReferenceRestrictions/Public/AssetReferencingPolicySettings.h
@@ -196,6 +196,10 @@ public:
 	UPROPERTY(config, EditAnywhere, Category="Project Content", meta=(TitleProperty=DomainName))
 	TArray<FARPDomainDefinitionByContentRoot> AdditionalDomains;
 
+	// If true, ignore any editor-only domain references that will not affect the cooked game
+	UPROPERTY(config, EditAnywhere, Category = Settings)
+	bool bIgnoreEditorOnlyReferences = true;
+
 	// The names of the project, special system mount, and game domains
 	static const FString EngineDomainName;
 	static const FString ScriptDomainName;
