Add info to ensure about domain path conflicts

diff --git a/Engine/Plugins/Editor/AssetReferenceRestrictions/Source/AssetReferenceRestrictions/Private/AssetReferencingDomains.cpp b/Engine/Plugins/Editor/AssetReferenceRestrictions/Source/AssetReferenceRestrictions/Private/AssetReferencingDomains.cpp
--- a/Engine/Plugins/Editor/AssetReferenceRestrictions/Source/AssetReferenceRestrictions/Private/AssetReferencingDomains.cpp
+++ b/Engine/Plugins/Editor/AssetReferenceRestrictions/Source/AssetReferenceRestrictions/Private/AssetReferencingDomains.cpp
@@ -27,8 +27,22 @@ public:
 
 	// The domain to use if nothing more specific in SubFolders matches
 	TSharedPtr<FDomainData> DefaultDomain;
+	
+	// Weak pointer to parent node for debugging 
+	TWeakPtr<FDomainPathNode> Parent;
+	
+	// Path segment for this node for debugging
+	FString PathSegment;
 
 public:
+	
+	FDomainPathNode(TWeakPtr<FDomainPathNode> InParent, FString InPathSegment)
+		: Parent(InParent)
+		, PathSegment(InPathSegment)
+	{
+
+	}
+
 	// RemainingPath is expected to have leading but no trailing /
 	TSharedPtr<FDomainData> FindDomainFromPath(FStringView RemainingPath) const
 	{
@@ -72,18 +86,35 @@ public:
 			TSharedPtr<FDomainPathNode>& ChildFolder = SubFolders.FindOrAdd(DirectoryName);
 			if (!ChildFolder.IsValid())
 			{
-				ChildFolder = MakeShared<FDomainPathNode>();
+				ChildFolder = MakeShared<FDomainPathNode>(this->AsWeak(), DirectoryName);
 			}
 
 			ChildFolder->AddDomain(Domain, RemainingPath.Mid(DirectorySeparatorIndex + 1));
 		}
 		else
 		{
+			UE_LOG(LogAssetReferenceRestrictions, Verbose, TEXT("Assigning path %s to domain %s"), *GetPath(), *Domain->UserFacingDomainName.ToString());
 			check(RemainingPath.Len() == 0);
-			ensure(!DefaultDomain.IsValid());
+			ensureMsgf(!DefaultDomain.IsValid(), TEXT("Registered the path %s for two different domains: %s and %s"), 
+				*GetPath(), 
+				*DefaultDomain->UserFacingDomainName.ToString(),
+				*Domain->UserFacingDomainName.ToString()
+				);
 			DefaultDomain = Domain;
 		}
 	}
+	
+	FString GetPath() const
+	{
+		if (TSharedPtr<FDomainPathNode> PinnedParent = Parent.Pin(); PinnedParent.IsValid())
+		{
+			return FString::Printf(TEXT("%s%s/"), *PinnedParent->GetPath(), *PathSegment);
+		}
+		else
+		{
+			return FString(TEXT("/"));
+		}
+	}
 
 	void DebugPrint(const FString& PreceedingPath, int32 Depth) const
 	{
@@ -234,7 +265,7 @@ void FDomainDatabase::RebuildFromScratch()
 	}
 
 	// Rebuild the path map
-	PathMap = MakeShared<FDomainPathNode>();
+	PathMap = MakeShared<FDomainPathNode>(nullptr, FString{});
 	for (const auto& KVP : DomainNameMap)
 	{
 		TSharedPtr<FDomainData> Domain = KVP.Value;
