Don't allow duplicate entries in never cook referencing domain to avoid ensure.

diff --git a/Engine/Plugins/Editor/AssetReferenceRestrictions/Source/AssetReferenceRestrictions/Private/AssetReferencingDomains.cpp b/Engine/Plugins/Editor/AssetReferenceRestrictions/Source/AssetReferenceRestrictions/Private/AssetReferencingDomains.cpp
--- a/Engine/Plugins/Editor/AssetReferenceRestrictions/Source/AssetReferenceRestrictions/Private/AssetReferencingDomains.cpp
+++ b/Engine/Plugins/Editor/AssetReferenceRestrictions/Source/AssetReferenceRestrictions/Private/AssetReferencingDomains.cpp
@@ -257,11 +257,13 @@ void FDomainDatabase::RebuildFromScratch()
 		NeverCookDomain->UserFacingDomainName = LOCTEXT("NeverCook", "Never Cooked Content");
 		NeverCookDomain->bCanSeeEverything = true;
 
+		TSet<FString> DomainRootPaths;
 		for (const FDirectoryPath& DirectoryToNeverCook : PackagingSettings->DirectoriesToNeverCook)
 		{
 			const FString UncookedFolder = DirectoryToNeverCook.Path.StartsWith(TEXT("/"), ESearchCase::CaseSensitive) ? DirectoryToNeverCook.Path : (TEXT("/Game/") + DirectoryToNeverCook.Path);
-			NeverCookDomain->DomainRootPaths.Add(UncookedFolder / TEXT(""));
+			DomainRootPaths.Add(UncookedFolder / TEXT(""));
 		}
+		NeverCookDomain->DomainRootPaths = DomainRootPaths.Array();
 	}
 
 	// Rebuild the path map
