Rebuild asset referencing domain at the beginning of asset validation so any DB errors aren't associated with an asset.

diff --git a/Engine/Plugins/Editor/AssetReferenceRestrictions/Source/AssetReferenceRestrictions/Private/AssetReferencingPolicySubsystem.cpp b/Engine/Plugins/Editor/AssetReferenceRestrictions/Source/AssetReferenceRestrictions/Private/AssetReferencingPolicySubsystem.cpp
--- a/Engine/Plugins/Editor/AssetReferenceRestrictions/Source/AssetReferenceRestrictions/Private/AssetReferencingPolicySubsystem.cpp
+++ b/Engine/Plugins/Editor/AssetReferenceRestrictions/Source/AssetReferenceRestrictions/Private/AssetReferencingPolicySubsystem.cpp
@@ -21,6 +21,8 @@ void UAssetReferencingPolicySubsystem::Initialize(FSubsystemCollectionBase& Coll
 
 	DomainDB = MakeShared<FDomainDatabase>();
 	DomainDB->Init();
+
+	FEditorDelegates::OnPreAssetValidation.AddUObject(this, &UAssetReferencingPolicySubsystem::UpdateDBIfNecessary);
 }
 
 void UAssetReferencingPolicySubsystem::Deinitialize()
@@ -35,6 +37,11 @@ TSharedPtr<IAssetReferenceFilter> UAssetReferencingPolicySubsystem::HandleMakeAs
 	return MakeShareable(new FDomainAssetReferenceFilter(Context, GetDomainDB()));
 }
 
+void UAssetReferencingPolicySubsystem::UpdateDBIfNecessary() const
+{
+	DomainDB->UpdateIfNecessary();
+}
+
 TSharedPtr<FDomainDatabase> UAssetReferencingPolicySubsystem::GetDomainDB() const
 {
 	DomainDB->UpdateIfNecessary();
diff --git a/Engine/Plugins/Editor/AssetReferenceRestrictions/Source/AssetReferenceRestrictions/Public/AssetReferencingPolicySubsystem.h b/Engine/Plugins/Editor/AssetReferenceRestrictions/Source/AssetReferenceRestrictions/Public/AssetReferencingPolicySubsystem.h
--- a/Engine/Plugins/Editor/AssetReferenceRestrictions/Source/AssetReferenceRestrictions/Public/AssetReferencingPolicySubsystem.h
+++ b/Engine/Plugins/Editor/AssetReferenceRestrictions/Source/AssetReferenceRestrictions/Public/AssetReferencingPolicySubsystem.h
@@ -25,6 +25,8 @@ public:
 
 	TSharedPtr<FDomainDatabase> GetDomainDB() const;
 private:
+	void UpdateDBIfNecessary() const;
+
 	TSharedPtr<IAssetReferenceFilter> HandleMakeAssetReferenceFilter(const FAssetReferenceFilterContext& Context);
 
 	TSharedPtr<FDomainDatabase> DomainDB;
