Build an unrestricted domain for gamefeature plugins that can't have content (roots)

diff --git a/Engine/Plugins/Editor/AssetReferenceRestrictions/Source/AssetReferenceRestrictions/AssetReferenceRestrictions.Build.cs b/Engine/Plugins/Editor/AssetReferenceRestrictions/Source/AssetReferenceRestrictions/AssetReferenceRestrictions.Build.cs
--- a/Engine/Plugins/Editor/AssetReferenceRestrictions/Source/AssetReferenceRestrictions/AssetReferenceRestrictions.Build.cs
+++ b/Engine/Plugins/Editor/AssetReferenceRestrictions/Source/AssetReferenceRestrictions/AssetReferenceRestrictions.Build.cs
@@ -24,7 +24,8 @@ namespace UnrealBuildTool.Rules
 					"UnrealEd",
 					"EditorFramework",
 					"Projects",
-					"DataValidation"
+					"DataValidation",
+					"Json"
 				}
 			);
 		}
diff --git a/Engine/Plugins/Editor/AssetReferenceRestrictions/Source/AssetReferenceRestrictions/Private/AssetReferencingDomains.cpp b/Engine/Plugins/Editor/AssetReferenceRestrictions/Source/AssetReferenceRestrictions/Private/AssetReferencingDomains.cpp
--- a/Engine/Plugins/Editor/AssetReferenceRestrictions/Source/AssetReferenceRestrictions/Private/AssetReferencingDomains.cpp
+++ b/Engine/Plugins/Editor/AssetReferenceRestrictions/Source/AssetReferenceRestrictions/Private/AssetReferencingDomains.cpp
@@ -10,6 +10,9 @@
 #include "Misc/Paths.h"
 #include "PluginReferenceDescriptor.h"
 
+#include "Dom/JsonObject.h"
+#include "Dom/JsonValue.h"
+
 DEFINE_LOG_CATEGORY(LogAssetReferenceRestrictions);
 
 #define LOCTEXT_NAMESPACE "AssetReferencingPolicy"
@@ -243,7 +246,16 @@ void FDomainDatabase::RebuildFromScratch()
 	TArray<TSharedRef<IPlugin>> EnabledPlugins = IPluginManager::Get().GetEnabledPlugins();
 	for (const TSharedRef<IPlugin>& Plugin : EnabledPlugins)
 	{
-		if (Plugin->CanContainContent())
+		// Temporary until NoGameFeatureContent is removed. Don't add dependencies for GFPs that have no content
+		// other than the GameFeatureData
+		bool bNoGameFeatureContent = false;
+		Plugin->GetDescriptorJson()->TryGetBoolField(TEXT("NoGameFeatureContent"), bNoGameFeatureContent);
+		
+		if (bNoGameFeatureContent)
+		{
+			BuildUnrestrictedDomainFromPlugin(Plugin);
+		}
+		else if (Plugin->CanContainContent())
 		{
 			BuildDomainFromPlugin(Plugin);
 		}
@@ -446,6 +458,18 @@ void FDomainDatabase::BuildDomainFromPlugin(TSharedRef<IPlugin> Plugin)
 	}
 }
 
+void FDomainDatabase::BuildUnrestrictedDomainFromPlugin(TSharedRef<IPlugin> Plugin)
+{
+	const FString NewDomainName = Plugin->GetName();
+	DomainsDefinedByPlugins.Add(NewDomainName);
+
+	TSharedPtr<FDomainData> Domain = FindOrAddDomainByName(NewDomainName);
+
+	Domain->Reset();
+	Domain->DomainRootPaths.Add(Plugin->GetMountedAssetPath());
+	Domain->bCanBeSeenByEverything = true;
+}
+
 void FDomainDatabase::ValidateAllDomains()
 {
 	for (const auto& KVP : DomainNameMap)
diff --git a/Engine/Plugins/Editor/AssetReferenceRestrictions/Source/AssetReferenceRestrictions/Public/AssetReferencingDomains.h b/Engine/Plugins/Editor/AssetReferenceRestrictions/Source/AssetReferenceRestrictions/Public/AssetReferencingDomains.h
--- a/Engine/Plugins/Editor/AssetReferenceRestrictions/Source/AssetReferenceRestrictions/Public/AssetReferencingDomains.h
+++ b/Engine/Plugins/Editor/AssetReferenceRestrictions/Source/AssetReferenceRestrictions/Public/AssetReferencingDomains.h
@@ -81,6 +81,7 @@ private:
 	void RebuildFromScratch();
 	
 	void BuildDomainFromPlugin(TSharedRef<IPlugin> Plugin);
+	void BuildUnrestrictedDomainFromPlugin(TSharedRef<IPlugin> Plugin);
 
 	void AddDomainVisibilityList(TSharedPtr<FDomainData> Domain, const TArray<FString>& VisibilityList);
 private:
