UAssetValidator_AssetReferenceRestrictions: Enabled validation on asset's associated external objects (to also validate external actors)

diff --git a/Engine/Plugins/Editor/AssetReferenceRestrictions/Source/AssetReferenceRestrictions/Private/AssetValidator_AssetReferenceRestrictions.cpp b/Engine/Plugins/Editor/AssetReferenceRestrictions/Source/AssetReferenceRestrictions/Private/AssetValidator_AssetReferenceRestrictions.cpp
--- a/Engine/Plugins/Editor/AssetReferenceRestrictions/Source/AssetReferenceRestrictions/Private/AssetValidator_AssetReferenceRestrictions.cpp
+++ b/Engine/Plugins/Editor/AssetReferenceRestrictions/Source/AssetReferenceRestrictions/Private/AssetValidator_AssetReferenceRestrictions.cpp
@@ -11,6 +11,7 @@
 #include "AssetReferencingDomains.h"
 #include "Editor/AssetReferenceFilter.h"
 #include "Misc/PackageName.h"
+#include "Misc/DataValidation.h"
 #include "Modules/ModuleManager.h"
 
 #include UE_INLINE_GENERATED_CPP_BY_NAME(AssetValidator_AssetReferenceRestrictions)
@@ -43,11 +44,31 @@ EDataValidationResult UAssetValidator_AssetReferenceRestrictions::ValidateLoaded
 	check(InAsset);
 
 	FAssetRegistryModule& AssetRegistryModule = FModuleManager::LoadModuleChecked<FAssetRegistryModule>(TEXT("AssetRegistry"));
-	IAssetRegistry& AssetRegistry = AssetRegistryModule.Get();
-	const FName TransientName = GetTransientPackage()->GetFName();
+	const IAssetRegistry& AssetRegistry = AssetRegistryModule.Get();
+
+	// Validate asset references
+	ValidateAssetInternal(InAssetData, AssetRegistry);
+	
+	// Validate each external object's references
+	for (const FAssetData& ExternalObject : InContext.GetAssociatedExternalObjects())
+	{
+		ValidateAssetInternal(ExternalObject, AssetRegistry);
+	}
+	
+	if (GetValidationResult() != EDataValidationResult::Invalid)
+	{
+		AssetPasses(InAsset);
+	}
+
+	return GetValidationResult();
+}
+
+void UAssetValidator_AssetReferenceRestrictions::ValidateAssetInternal(const FAssetData& InAssetData, const IAssetRegistry& InAssetRegistry)
+{
+	static const FName TransientName = GetTransientPackage()->GetFName();
 
 	// Check for missing soft or hard references to cinematic and developers content
-	FName PackageFName = InAsset->GetOutermost()->GetFName();
+	FName PackageFName = InAssetData.PackageName;
 	TArray<FName> SoftDependencies;
 	TArray<FAssetData> AllDependencyAssets;
 
@@ -60,19 +81,19 @@ EDataValidationResult UAssetValidator_AssetReferenceRestrictions::ValidateLoaded
 		QueryFlags = UE::AssetRegistry::EDependencyQuery::Game;
 	}
 
-	AssetRegistry.GetDependencies(PackageFName, SoftDependencies, UE::AssetRegistry::EDependencyCategory::Package, UE::AssetRegistry::EDependencyQuery::Soft | QueryFlags);
+	InAssetRegistry.GetDependencies(PackageFName, SoftDependencies, UE::AssetRegistry::EDependencyCategory::Package, UE::AssetRegistry::EDependencyQuery::Soft | QueryFlags);
 	for (FName SoftDependency : SoftDependencies)
 	{
 		const FString SoftDependencyStr = SoftDependency.ToString();
 		if (!FPackageName::IsScriptPackage(SoftDependencyStr))
 		{
 			TArray<FAssetData> DependencyAssets;
-			AssetRegistry.GetAssetsByPackageName(SoftDependency, DependencyAssets, true);
+			InAssetRegistry.GetAssetsByPackageName(SoftDependency, DependencyAssets, true);
 			if (DependencyAssets.Num() == 0)
 			{
 				if (SoftDependency != TransientName)
 				{
-					AssetFails(InAsset, FText::Format(LOCTEXT("IllegalReference_MissingSoftRef", "Soft references {0} which does not exist"), FText::FromString(SoftDependencyStr)));
+					AssetMessage(InAssetData, EMessageSeverity::Error, FText::Format(LOCTEXT("IllegalReference_MissingSoftRef", "Soft references {0} which does not exist"), FText::FromString(SoftDependencyStr)));
 				}
 			}
 			else
@@ -84,7 +105,7 @@ EDataValidationResult UAssetValidator_AssetReferenceRestrictions::ValidateLoaded
 
 	// Now check hard references to cinematic and developers content
 	TArray<FName> HardDependencies;
-	AssetRegistry.GetDependencies(PackageFName, HardDependencies, UE::AssetRegistry::EDependencyCategory::Package, UE::AssetRegistry::EDependencyQuery::Hard | QueryFlags);
+	InAssetRegistry.GetDependencies(PackageFName, HardDependencies, UE::AssetRegistry::EDependencyCategory::Package, UE::AssetRegistry::EDependencyQuery::Hard | QueryFlags);
 	for (FName HardDependency : HardDependencies)
 	{
 		//@TODO: Probably not needed anymore?
@@ -93,11 +114,11 @@ EDataValidationResult UAssetValidator_AssetReferenceRestrictions::ValidateLoaded
 		FString UncookedFolderName;
 		if (IsInUncookedFolder(HardDependencyStr, &UncookedFolderName))
 		{
-			AssetFails(InAsset, FText::Format(LOCTEXT("IllegalReference_HardDependency", "Illegally hard references {0} asset {1}"), FText::FromString(UncookedFolderName), FText::FromString(HardDependencyStr)));
+			AssetMessage(InAssetData, EMessageSeverity::Error, FText::Format(LOCTEXT("IllegalReference_HardDependency", "Illegally hard references {0} asset {1}"), FText::FromString(UncookedFolderName), FText::FromString(HardDependencyStr)));
 		}
 #endif
 
-		AssetRegistry.GetAssetsByPackageName(HardDependency, AllDependencyAssets, true);
+		InAssetRegistry.GetAssetsByPackageName(HardDependency, AllDependencyAssets, true);
 	}
 
 	if ((GetValidationResult() != EDataValidationResult::Invalid) && (AllDependencyAssets.Num() > 0))
@@ -112,20 +133,13 @@ EDataValidationResult UAssetValidator_AssetReferenceRestrictions::ValidateLoaded
 				FText FailureReason;
 				if (!AssetReferenceFilter->PassesFilter(Dependency, &FailureReason))
 				{
-					AssetMessage(InAsset, EMessageSeverity::Error, FText::Format(LOCTEXT("IllegalReference_AssetFilterFail", "Illegal reference:"), FailureReason))
+					AssetMessage(InAssetData, EMessageSeverity::Error, FText::Format(LOCTEXT("IllegalReference_AssetFilterFail", "Illegal reference:"), FailureReason))
 						->AddToken(FAssetDataToken::Create(Dependency))
 						->AddText(FText::Format(LOCTEXT("IllegalReference_FailureReason", ". {0}"), FailureReason));
 				}
 			}
 		}
 	}
-
-	if (GetValidationResult() != EDataValidationResult::Invalid)
-	{
-		AssetPasses(InAsset);
-	}
-
-	return GetValidationResult();
 }
 
 #undef LOCTEXT_NAMESPACE
diff --git a/Engine/Plugins/Editor/AssetReferenceRestrictions/Source/AssetReferenceRestrictions/Private/AssetValidator_AssetReferenceRestrictions.h b/Engine/Plugins/Editor/AssetReferenceRestrictions/Source/AssetReferenceRestrictions/Private/AssetValidator_AssetReferenceRestrictions.h
--- a/Engine/Plugins/Editor/AssetReferenceRestrictions/Source/AssetReferenceRestrictions/Private/AssetValidator_AssetReferenceRestrictions.h
+++ b/Engine/Plugins/Editor/AssetReferenceRestrictions/Source/AssetReferenceRestrictions/Private/AssetValidator_AssetReferenceRestrictions.h
@@ -5,6 +5,8 @@
 #include "EditorValidatorBase.h"
 #include "AssetValidator_AssetReferenceRestrictions.generated.h"
 
+class IAssetRegistry;
+
 UCLASS()
 class UAssetValidator_AssetReferenceRestrictions : public UEditorValidatorBase
 {
@@ -18,4 +20,7 @@ protected:
 	virtual bool CanValidateAsset_Implementation(const FAssetData& InAssetData, UObject* InAsset, FDataValidationContext& InContext) const override;
 	virtual EDataValidationResult ValidateLoadedAsset_Implementation(const FAssetData& InAssetData, UObject* InAsset, FDataValidationContext& InContext) override;
 	//~End of UEditorValidatorBase interface
+
+private:
+	void ValidateAssetInternal(const FAssetData& InAssetData, const IAssetRegistry& InAssetRegistry);
 };
