Fixed issue where setting a Verse class in a TSoftClassPtr property on a BP asset would cause validation errors.

diff --git a/Engine/Plugins/Editor/AssetReferenceRestrictions/Source/AssetReferenceRestrictions/Private/AssetValidator_AssetReferenceRestrictions.cpp b/Engine/Plugins/Editor/AssetReferenceRestrictions/Source/AssetReferenceRestrictions/Private/AssetValidator_AssetReferenceRestrictions.cpp
--- a/Engine/Plugins/Editor/AssetReferenceRestrictions/Source/AssetReferenceRestrictions/Private/AssetValidator_AssetReferenceRestrictions.cpp
+++ b/Engine/Plugins/Editor/AssetReferenceRestrictions/Source/AssetReferenceRestrictions/Private/AssetValidator_AssetReferenceRestrictions.cpp
@@ -85,7 +85,7 @@ void UAssetValidator_AssetReferenceRestrictions::ValidateAssetInternal(const FAs
 	for (FName SoftDependency : SoftDependencies)
 	{
 		const FString SoftDependencyStr = SoftDependency.ToString();
-		if (!FPackageName::IsScriptPackage(SoftDependencyStr))
+		if (!FPackageName::IsScriptPackage(SoftDependencyStr) && !FPackageName::IsVersePackage(SoftDependencyStr))
 		{
 			TArray<FAssetData> DependencyAssets;
 			InAssetRegistry.GetAssetsByPackageName(SoftDependency, DependencyAssets, true);
