Validator to check asset referencing restrictions for assets in plugins when plugins descriptor changes.

diff --git a/Engine/Plugins/Editor/AssetReferenceRestrictions/Source/AssetReferenceRestrictions/Private/AssetReferencingPolicySubsystem.cpp b/Engine/Plugins/Editor/AssetReferenceRestrictions/Source/AssetReferenceRestrictions/Private/AssetReferencingPolicySubsystem.cpp
--- a/Engine/Plugins/Editor/AssetReferenceRestrictions/Source/AssetReferenceRestrictions/Private/AssetReferencingPolicySubsystem.cpp
+++ b/Engine/Plugins/Editor/AssetReferenceRestrictions/Source/AssetReferenceRestrictions/Private/AssetReferencingPolicySubsystem.cpp
@@ -1,9 +1,13 @@
 // Copyright Epic Games, Inc. All Rights Reserved.
 
 #include "AssetReferencingPolicySubsystem.h"
+
+#include "AssetReferencingDomains.h"
+#include "AssetReferencingPolicySettings.h"
+#include "AssetRegistry/AssetRegistryModule.h"
 #include "DomainAssetReferenceFilter.h"
 #include "Editor.h"
-#include "AssetReferencingDomains.h"
+#include "Misc/PathViews.h"
 
 #include UE_INLINE_GENERATED_CPP_BY_NAME(AssetReferencingPolicySubsystem)
 
@@ -32,6 +36,100 @@ void UAssetReferencingPolicySubsystem::Deinitialize()
 	DomainDB.Reset();
 }
 
+bool UAssetReferencingPolicySubsystem::ShouldValidateAssetReferences(const FAssetData& Asset) const
+{
+	TSharedPtr<FDomainData> DomainData = GetDomainDB()->FindDomainFromAssetData(Asset);
+
+	const bool bIsInUnrestrictedFolder = DomainData && DomainData->IsValid() && DomainData->bCanSeeEverything;
+	return !bIsInUnrestrictedFolder;
+}
+
+TValueOrError<void, TArray<FAssetReferenceError>> UAssetReferencingPolicySubsystem::ValidateAssetReferences(const FAssetData& InAssetData) const
+{
+	IAssetRegistry& AssetRegistry = FModuleManager::LoadModuleChecked<FAssetRegistryModule>("AssetRegistry").Get();
+	TArray<FAssetReferenceError> Errors;
+	static const FName TransientName = GetTransientPackage()->GetFName();
+
+	// Check for missing soft or hard references to cinematic and developers content
+	FName PackageFName = InAssetData.PackageName;
+	TArray<FName> SoftDependencies;
+	TArray<FAssetData> AllDependencyAssets;
+
+	const UAssetReferencingPolicySettings* Settings = GetDefault<UAssetReferencingPolicySettings>();
+	UE::AssetRegistry::EDependencyQuery QueryFlags = UE::AssetRegistry::EDependencyQuery::NoRequirements;
+
+	if (Settings->bIgnoreEditorOnlyReferences)
+	{
+		// If the rules allow ignoring editor-only referencers, restrict the dependencies to game references
+		QueryFlags = UE::AssetRegistry::EDependencyQuery::Game;
+	}
+
+	AssetRegistry.GetDependencies(PackageFName, SoftDependencies, UE::AssetRegistry::EDependencyCategory::Package, UE::AssetRegistry::EDependencyQuery::Soft | QueryFlags);
+	for (FName SoftDependency : SoftDependencies)
+	{
+		FNameBuilder SoftDependencyStr(SoftDependency);
+		if (!FPackageName::IsScriptPackage(SoftDependencyStr.ToView()) && !FPackageName::IsVersePackage(SoftDependencyStr.ToView()))
+		{
+			TArray<FAssetData> DependencyAssets;
+			AssetRegistry.GetAssetsByPackageName(SoftDependency, DependencyAssets, true);
+			if (DependencyAssets.Num() == 0)
+			{
+				if (SoftDependency != TransientName)
+				{
+					Errors.Add(FAssetReferenceError{
+						.Type = EAssetReferenceErrorType::DoesNotExist,
+						.ReferencedAsset = FAssetData(SoftDependency, FName(FPathViews::GetPath(SoftDependencyStr)), FName(), FTopLevelAssetPath(UObject::StaticClass())),
+						.Message = FText::Format(LOCTEXT("IllegalReference_MissingSoftRef", "Soft references {0} which does not exist"), FText::FromName(SoftDependency)),
+					});
+				}
+			}
+			else
+			{
+				AllDependencyAssets.Append(DependencyAssets);
+			}
+		}
+	}
+
+	// Now check hard references to cinematic and developers content
+	TArray<FName> HardDependencies;
+	AssetRegistry.GetDependencies(PackageFName, HardDependencies, UE::AssetRegistry::EDependencyCategory::Package, UE::AssetRegistry::EDependencyQuery::Hard | QueryFlags);
+	for (FName HardDependency : HardDependencies)
+	{
+		AssetRegistry.GetAssetsByPackageName(HardDependency, AllDependencyAssets, true);
+	}
+
+	if (AllDependencyAssets.Num() > 0)
+	{
+		FAssetReferenceFilterContext AssetReferenceFilterContext;
+		AssetReferenceFilterContext.ReferencingAssets = { InAssetData };
+		TSharedPtr<IAssetReferenceFilter> AssetReferenceFilter = GEditor ? GEditor->MakeAssetReferenceFilter(AssetReferenceFilterContext) : nullptr;
+		if (ensure(AssetReferenceFilter.IsValid()))
+		{
+			for (const FAssetData& Dependency : AllDependencyAssets)
+			{
+				FText FailureReason;
+				if (!AssetReferenceFilter->PassesFilter(Dependency, &FailureReason))
+				{
+					Errors.Add(FAssetReferenceError{
+						.Type = EAssetReferenceErrorType::Illegal,
+						.ReferencedAsset = Dependency,
+						.Message = FText::Format(LOCTEXT("IllegalReference_AssetFilterFail", "Illegal reference: {0}"), FailureReason),
+					});
+				}
+			}
+		}
+	}
+	
+	if (Errors.Num())
+	{
+		return MakeError(MoveTemp(Errors));
+	}
+	else
+	{
+		return MakeValue();
+	}
+}
+
 TSharedPtr<IAssetReferenceFilter> UAssetReferencingPolicySubsystem::HandleMakeAssetReferenceFilter(const FAssetReferenceFilterContext& Context)
 {
 	return MakeShareable(new FDomainAssetReferenceFilter(Context, GetDomainDB()));
diff --git a/Engine/Plugins/Editor/AssetReferenceRestrictions/Source/AssetReferenceRestrictions/Private/AssetValidator_AssetReferenceRestrictions.cpp b/Engine/Plugins/Editor/AssetReferenceRestrictions/Source/AssetReferenceRestrictions/Private/AssetValidator_AssetReferenceRestrictions.cpp
--- a/Engine/Plugins/Editor/AssetReferenceRestrictions/Source/AssetReferenceRestrictions/Private/AssetValidator_AssetReferenceRestrictions.cpp
+++ b/Engine/Plugins/Editor/AssetReferenceRestrictions/Source/AssetReferenceRestrictions/Private/AssetValidator_AssetReferenceRestrictions.cpp
@@ -27,13 +27,7 @@ bool UAssetValidator_AssetReferenceRestrictions::CanValidateAsset_Implementation
 {
 	if (InAsset)
 	{
-		TSharedPtr<FDomainData> DomainData = GEditor->GetEditorSubsystem<UAssetReferencingPolicySubsystem>()->GetDomainDB()->FindDomainFromAssetData(AssetData);
-
-		const bool bIsInUnrestrictedFolder = DomainData && DomainData->IsValid() && DomainData->bCanSeeEverything;
-		if (!bIsInUnrestrictedFolder)
-		{
-			return true;
-		}
+		return GEditor->GetEditorSubsystem<UAssetReferencingPolicySubsystem>()->ShouldValidateAssetReferences(AssetData);
 	}
 
 	return false;
@@ -65,79 +59,13 @@ EDataValidationResult UAssetValidator_AssetReferenceRestrictions::ValidateLoaded
 
 void UAssetValidator_AssetReferenceRestrictions::ValidateAssetInternal(const FAssetData& InAssetData, const IAssetRegistry& InAssetRegistry)
 {
-	static const FName TransientName = GetTransientPackage()->GetFName();
-
-	// Check for missing soft or hard references to cinematic and developers content
-	FName PackageFName = InAssetData.PackageName;
-	TArray<FName> SoftDependencies;
-	TArray<FAssetData> AllDependencyAssets;
-
-	const UAssetReferencingPolicySettings* Settings = GetDefault<UAssetReferencingPolicySettings>();
-	UE::AssetRegistry::EDependencyQuery QueryFlags = UE::AssetRegistry::EDependencyQuery::NoRequirements;
-
-	if (Settings->bIgnoreEditorOnlyReferences)
-	{
-		// If the rules allow ignoring editor-only referencers, restrict the dependencies to game references
-		QueryFlags = UE::AssetRegistry::EDependencyQuery::Game;
-	}
-
-	InAssetRegistry.GetDependencies(PackageFName, SoftDependencies, UE::AssetRegistry::EDependencyCategory::Package, UE::AssetRegistry::EDependencyQuery::Soft | QueryFlags);
-	for (FName SoftDependency : SoftDependencies)
-	{
-		const FString SoftDependencyStr = SoftDependency.ToString();
-		if (!FPackageName::IsScriptPackage(SoftDependencyStr) && !FPackageName::IsVersePackage(SoftDependencyStr))
-		{
-			TArray<FAssetData> DependencyAssets;
-			InAssetRegistry.GetAssetsByPackageName(SoftDependency, DependencyAssets, true);
-			if (DependencyAssets.Num() == 0)
-			{
-				if (SoftDependency != TransientName)
-				{
-					AssetMessage(InAssetData, EMessageSeverity::Error, FText::Format(LOCTEXT("IllegalReference_MissingSoftRef", "Soft references {0} which does not exist"), FText::FromString(SoftDependencyStr)));
-				}
-			}
-			else
-			{
-				AllDependencyAssets.Append(DependencyAssets);
-			}
-		}
-	}
-
-	// Now check hard references to cinematic and developers content
-	TArray<FName> HardDependencies;
-	InAssetRegistry.GetDependencies(PackageFName, HardDependencies, UE::AssetRegistry::EDependencyCategory::Package, UE::AssetRegistry::EDependencyQuery::Hard | QueryFlags);
-	for (FName HardDependency : HardDependencies)
-	{
-		//@TODO: Probably not needed anymore?
-#if 0
-		const FString HardDependencyStr = HardDependency.ToString();
-		FString UncookedFolderName;
-		if (IsInUncookedFolder(HardDependencyStr, &UncookedFolderName))
-		{
-			AssetMessage(InAssetData, EMessageSeverity::Error, FText::Format(LOCTEXT("IllegalReference_HardDependency", "Illegally hard references {0} asset {1}"), FText::FromString(UncookedFolderName), FText::FromString(HardDependencyStr)));
-		}
-#endif
-
-		InAssetRegistry.GetAssetsByPackageName(HardDependency, AllDependencyAssets, true);
-	}
-
-	if ((GetValidationResult() != EDataValidationResult::Invalid) && (AllDependencyAssets.Num() > 0))
+    UAssetReferencingPolicySubsystem* Subsystem = GEditor->GetEditorSubsystem<UAssetReferencingPolicySubsystem>();
+	TValueOrError<void, TArray<FAssetReferenceError>> Result = Subsystem->ValidateAssetReferences(InAssetData);
+	if (Result.HasError())
 	{
-		FAssetReferenceFilterContext AssetReferenceFilterContext;
-		AssetReferenceFilterContext.ReferencingAssets = { InAssetData };
-		TSharedPtr<IAssetReferenceFilter> AssetReferenceFilter = GEditor ? GEditor->MakeAssetReferenceFilter(AssetReferenceFilterContext) : nullptr;
-		if (ensure(AssetReferenceFilter.IsValid()))
+		for (const FAssetReferenceError& Error : Result.GetError())
 		{
-			for (const FAssetData& Dependency : AllDependencyAssets)
-			{
-				FText FailureReason;
-				if (!AssetReferenceFilter->PassesFilter(Dependency, &FailureReason))
-				{
-					AssetMessage(InAssetData, EMessageSeverity::Error, FText::Format(LOCTEXT("IllegalReference_AssetFilterFail", "Illegal reference:"), FailureReason))
-						->AddToken(FAssetDataToken::Create(Dependency))
-						->AddText(FText::Format(LOCTEXT("IllegalReference_FailureReason", ". {0}"), FailureReason));
-				}
-			}
+			AssetMessage(InAssetData, EMessageSeverity::Error, Error.Message)->AddToken(FAssetDataToken::Create(Error.ReferencedAsset));
 		}
 	}
 }
diff --git a/Engine/Plugins/Editor/AssetReferenceRestrictions/Source/AssetReferenceRestrictions/Private/EditorValidator_PluginAssetReferences.cpp b/Engine/Plugins/Editor/AssetReferenceRestrictions/Source/AssetReferenceRestrictions/Private/EditorValidator_PluginAssetReferences.cpp
new file mode 100644
--- /dev/null
+++ b/Engine/Plugins/Editor/AssetReferenceRestrictions/Source/AssetReferenceRestrictions/Private/EditorValidator_PluginAssetReferences.cpp
@@ -0,0 +1,84 @@
+// Copyright Epic Games, Inc. All Rights Reserved.
+
+#include "EditorValidator_PluginAssetReferences.h"
+
+#include "Algo/RemoveIf.h"
+#include "AssetReferencingPolicySubsystem.h"
+#include "AssetRegistry/AssetData.h"
+#include "AssetRegistry/AssetDataToken.h"
+#include "AssetRegistry/AssetRegistryModule.h"
+#include "DataValidationChangelist.h"
+#include "Editor.h"
+#include "Editor/EditorEngine.h"
+#include "Interfaces/IPluginManager.h"
+#include "Misc/PathViews.h"
+
+bool UEditorValidator_PluginAssetReferences::CanValidateAsset_Implementation(const FAssetData& InAssetData, UObject* InObject, FDataValidationContext& InContext) const
+{
+    if (UDataValidationChangelist* Changelist = Cast<UDataValidationChangelist>(InObject))
+    {
+        for (const FString& ModifiedFilePath : Changelist->ModifiedFiles)
+        {
+            if (FPathViews::GetExtension(ModifiedFilePath) == TEXTVIEW("uplugin"))
+            {
+                return true;
+            }
+        }
+    }
+    return false;
+}
+
+EDataValidationResult UEditorValidator_PluginAssetReferences::ValidateLoadedAsset_Implementation(const FAssetData& InAssetData, UObject* InAsset, FDataValidationContext& Context) 
+{
+    UDataValidationChangelist* Changelist = Cast<UDataValidationChangelist>(InAsset);
+    if (!Changelist)
+    {
+        return EDataValidationResult::Valid;
+    }
+
+    IPluginManager& PluginManager = IPluginManager::Get();
+    TArray<TSharedRef<IPlugin>> PluginsToValidate;
+    for (const FString& ModifiedFilePath : Changelist->ModifiedFiles)
+    {
+        if (FPathViews::GetExtension(ModifiedFilePath) == TEXTVIEW("uplugin"))
+        {
+            FStringView PluginName = FPathViews::GetBaseFilename(ModifiedFilePath);
+            TSharedPtr<IPlugin> Plugin = PluginManager.FindPlugin(PluginName);
+            if (Plugin.IsValid() && Plugin->IsEnabled())
+            {
+                PluginsToValidate.Add(Plugin.ToSharedRef());
+            }
+        }
+    }
+
+	IAssetRegistry& AssetRegistry = FModuleManager::LoadModuleChecked<FAssetRegistryModule>("AssetRegistry").Get();
+    TArray<FAssetData> AssetsToValidate;
+    for (const TSharedRef<IPlugin> Plugin : PluginsToValidate)
+    {
+        TStringBuilder<256> RootPath(InPlace, TEXT("/"), Plugin->GetName());
+        AssetRegistry.GetAssetsByPath(FName(RootPath), AssetsToValidate, true, true);
+    }
+
+    UAssetReferencingPolicySubsystem* Subsystem = GEditor->GetEditorSubsystem<UAssetReferencingPolicySubsystem>();
+    AssetsToValidate.SetNum(Algo::RemoveIf(AssetsToValidate, [Subsystem](const FAssetData& Asset) {
+        return !Subsystem->ShouldValidateAssetReferences(Asset);
+    }));
+
+    for (const FAssetData& Asset : AssetsToValidate)
+    {
+        TValueOrError<void, TArray<FAssetReferenceError>> Result = Subsystem->ValidateAssetReferences(Asset);
+        if (Result.HasError())
+        {
+            for (const FAssetReferenceError& Error : Result.GetError())
+            {
+                AssetMessage(Asset, EMessageSeverity::Error, Error.Message)->AddToken(FAssetDataToken::Create(Error.ReferencedAsset));
+            }
+        }
+    }
+
+    if(GetValidationResult() != EDataValidationResult::Invalid)
+    {
+        AssetPasses(InAsset);
+    }
+    return GetValidationResult();
+}
\ No newline at end of file
diff --git a/Engine/Plugins/Editor/AssetReferenceRestrictions/Source/AssetReferenceRestrictions/Private/EditorValidator_PluginAssetReferences.h b/Engine/Plugins/Editor/AssetReferenceRestrictions/Source/AssetReferenceRestrictions/Private/EditorValidator_PluginAssetReferences.h
new file mode 100644
--- /dev/null
+++ b/Engine/Plugins/Editor/AssetReferenceRestrictions/Source/AssetReferenceRestrictions/Private/EditorValidator_PluginAssetReferences.h
@@ -0,0 +1,17 @@
+// Copyright Epic Games, Inc. All Rights Reserved.
+
+#pragma once
+
+#include "EditorValidatorBase.h"
+
+#include "EditorValidator_PluginAssetReferences.generated.h"
+
+UCLASS()
+class UEditorValidator_PluginAssetReferences : public UEditorValidatorBase
+{
+    GENERATED_BODY()
+
+public:
+	virtual bool CanValidateAsset_Implementation(const FAssetData& InAssetData, UObject* InObject, FDataValidationContext& InContext) const override;
+	virtual EDataValidationResult ValidateLoadedAsset_Implementation(const FAssetData& InAssetData, UObject* InAsset, FDataValidationContext& Context) override;
+};
\ No newline at end of file
diff --git a/Engine/Plugins/Editor/AssetReferenceRestrictions/Source/AssetReferenceRestrictions/Public/AssetReferencingPolicySubsystem.h b/Engine/Plugins/Editor/AssetReferenceRestrictions/Source/AssetReferenceRestrictions/Public/AssetReferencingPolicySubsystem.h
--- a/Engine/Plugins/Editor/AssetReferenceRestrictions/Source/AssetReferenceRestrictions/Public/AssetReferencingPolicySubsystem.h
+++ b/Engine/Plugins/Editor/AssetReferenceRestrictions/Source/AssetReferenceRestrictions/Public/AssetReferencingPolicySubsystem.h
@@ -2,14 +2,30 @@
 
 #pragma once
 
+#include "AssetRegistry/AssetData.h"
 #include "EditorSubsystem.h"
+#include "Templates/ValueOrError.h"
 
 #include "AssetReferencingPolicySubsystem.generated.h"
 
 struct FAssetReferenceFilterContext;
 class IAssetReferenceFilter;
+struct FAssetData;
 struct FDomainDatabase;
 
+enum class EAssetReferenceErrorType
+{
+	DoesNotExist,
+	Illegal
+};
+
+struct FAssetReferenceError
+{
+	EAssetReferenceErrorType Type;
+	FAssetData ReferencedAsset;
+	FText Message;
+};
+
 /** Subsystem to register the domain-based asset referencing policy restrictions with the editor */
 UCLASS()
 class ASSETREFERENCERESTRICTIONS_API UAssetReferencingPolicySubsystem : public UEditorSubsystem
@@ -23,6 +39,12 @@ public:
 	virtual void Deinitialize() override;
 	//~End of UEditorSubsystem interface
 
+	// Returns whether the given asset's outgoing references are restricted in any way and should be individually validated
+	bool ShouldValidateAssetReferences(const FAssetData& Asset) const;
+
+	// Check the outgoing references of the given asset according to the asset registry and return details of any errors 
+	TValueOrError<void, TArray<FAssetReferenceError>> ValidateAssetReferences(const FAssetData& Asset) const;
+
 	TSharedPtr<FDomainDatabase> GetDomainDB() const;
 private:
 	void UpdateDBIfNecessary() const;
