Fix copy of variable in ranged-for

diff --git a/Engine/Plugins/Editor/AssetReferenceRestrictions/Source/AssetReferenceRestrictions/Private/EditorValidator_PluginAssetReferences.cpp b/Engine/Plugins/Editor/AssetReferenceRestrictions/Source/AssetReferenceRestrictions/Private/EditorValidator_PluginAssetReferences.cpp
--- a/Engine/Plugins/Editor/AssetReferenceRestrictions/Source/AssetReferenceRestrictions/Private/EditorValidator_PluginAssetReferences.cpp
+++ b/Engine/Plugins/Editor/AssetReferenceRestrictions/Source/AssetReferenceRestrictions/Private/EditorValidator_PluginAssetReferences.cpp
@@ -53,7 +53,7 @@ EDataValidationResult UEditorValidator_PluginAssetReferences::ValidateLoadedAsse
 
 	IAssetRegistry& AssetRegistry = FModuleManager::LoadModuleChecked<FAssetRegistryModule>("AssetRegistry").Get();
     TArray<FAssetData> AssetsToValidate;
-    for (const TSharedRef<IPlugin> Plugin : PluginsToValidate)
+    for (const TSharedRef<IPlugin>& Plugin : PluginsToValidate)
     {
         TStringBuilder<256> RootPath(InPlace, TEXT("/"), Plugin->GetName());
         AssetRegistry.GetAssetsByPath(FName(RootPath), AssetsToValidate, true, true);
