#UE Ask the EditorValidatorSubsystem if it is allowed to validate the files in a plugin when validating a uplugin change

diff --git a/Engine/Plugins/Editor/AssetReferenceRestrictions/Source/AssetReferenceRestrictions/Private/EditorValidator_PluginAssetReferences.cpp b/Engine/Plugins/Editor/AssetReferenceRestrictions/Source/AssetReferenceRestrictions/Private/EditorValidator_PluginAssetReferences.cpp
--- a/Engine/Plugins/Editor/AssetReferenceRestrictions/Source/AssetReferenceRestrictions/Private/EditorValidator_PluginAssetReferences.cpp
+++ b/Engine/Plugins/Editor/AssetReferenceRestrictions/Source/AssetReferenceRestrictions/Private/EditorValidator_PluginAssetReferences.cpp
@@ -9,6 +9,7 @@
 #include "AssetRegistry/AssetRegistryModule.h"
 #include "DataValidationChangelist.h"
 #include "Editor.h"
+#include "EditorValidatorSubsystem.h"
 #include "Editor/EditorEngine.h"
 #include "Interfaces/IPluginManager.h"
 #include "Misc/PathViews.h"
@@ -59,14 +60,17 @@ EDataValidationResult UEditorValidator_PluginAssetReferences::ValidateLoadedAsse
         AssetRegistry.GetAssetsByPath(FName(RootPath), AssetsToValidate, true, true);
     }
 
-    UAssetReferencingPolicySubsystem* Subsystem = GEditor->GetEditorSubsystem<UAssetReferencingPolicySubsystem>();
-    AssetsToValidate.SetNum(Algo::RemoveIf(AssetsToValidate, [Subsystem](const FAssetData& Asset) {
-        return !Subsystem->ShouldValidateAssetReferences(Asset);
+    UAssetReferencingPolicySubsystem* AssetReferencingPolicySubsystem = GEditor->GetEditorSubsystem<UAssetReferencingPolicySubsystem>();
+	UEditorValidatorSubsystem* EditorValidationSubsystem = GEditor->GetEditorSubsystem<UEditorValidatorSubsystem>();
+	FValidateAssetsSettings Settings;
+	FDataValidationContext ValidationContext(false, Settings.ValidationUsecase, {});
+    AssetsToValidate.SetNum(Algo::RemoveIf(AssetsToValidate, [AssetReferencingPolicySubsystem, EditorValidationSubsystem, &Settings, &ValidationContext](const FAssetData& Asset) {
+        return !AssetReferencingPolicySubsystem->ShouldValidateAssetReferences(Asset) || !EditorValidationSubsystem->ShouldValidateAsset(Asset, Settings, ValidationContext);
     }));
 
     for (const FAssetData& Asset : AssetsToValidate)
     {
-        TValueOrError<void, TArray<FAssetReferenceError>> Result = Subsystem->ValidateAssetReferences(Asset);
+        TValueOrError<void, TArray<FAssetReferenceError>> Result = AssetReferencingPolicySubsystem->ValidateAssetReferences(Asset);
         if (Result.HasError())
         {
             for (const FAssetReferenceError& Error : Result.GetError())
