Add fixer to illegal reference data validation errors

Added IAssetReferenceFixer (following the pattern used for IAssetReferenceFilter), which allows UAssetValidator_AssetReferenceRestrictions to add a fixer token to any illegal references it finds.

diff --git a/Engine/Plugins/Editor/AssetReferenceRestrictions/Source/AssetReferenceRestrictions/Private/AssetValidator_AssetReferenceRestrictions.cpp b/Engine/Plugins/Editor/AssetReferenceRestrictions/Source/AssetReferenceRestrictions/Private/AssetValidator_AssetReferenceRestrictions.cpp
--- a/Engine/Plugins/Editor/AssetReferenceRestrictions/Source/AssetReferenceRestrictions/Private/AssetValidator_AssetReferenceRestrictions.cpp
+++ b/Engine/Plugins/Editor/AssetReferenceRestrictions/Source/AssetReferenceRestrictions/Private/AssetValidator_AssetReferenceRestrictions.cpp
@@ -6,12 +6,14 @@
 #include "AssetRegistry/AssetDataToken.h"
 #include "AssetRegistry/AssetRegistryModule.h"
 #include "Editor.h"
+#include "Editor/AssetReferenceFixer.h"
 #include "AssetReferencingPolicySubsystem.h"
 #include "AssetReferencingPolicySettings.h"
 #include "AssetReferencingDomains.h"
 #include "Editor/AssetReferenceFilter.h"
 #include "Misc/PackageName.h"
 #include "Misc/DataValidation.h"
+#include "Misc/DataValidation/Fixer.h"
 #include "Modules/ModuleManager.h"
 
 #include UE_INLINE_GENERATED_CPP_BY_NAME(AssetValidator_AssetReferenceRestrictions)
@@ -60,12 +62,28 @@ EDataValidationResult UAssetValidator_AssetReferenceRestrictions::ValidateLoaded
 void UAssetValidator_AssetReferenceRestrictions::ValidateAssetInternal(const FAssetData& InAssetData, const IAssetRegistry& InAssetRegistry)
 {
     UAssetReferencingPolicySubsystem* Subsystem = GEditor->GetEditorSubsystem<UAssetReferencingPolicySubsystem>();
-	TValueOrError<void, TArray<FAssetReferenceError>> Result = Subsystem->ValidateAssetReferences(InAssetData);
-	if (Result.HasError())
+	if (TValueOrError<void, TArray<FAssetReferenceError>> Result = Subsystem->ValidateAssetReferences(InAssetData);
+		Result.HasError())
 	{
+		TSharedPtr<IAssetReferenceFixer> AssetReferenceFixer;
 		for (const FAssetReferenceError& Error : Result.GetError())
 		{
-			AssetMessage(InAssetData, EMessageSeverity::Error, Error.Message)->AddToken(FAssetDataToken::Create(Error.ReferencedAsset));
+			TSharedRef<FTokenizedMessage> TokenizedMessage = AssetMessage(InAssetData, EMessageSeverity::Error, Error.Message)->AddToken(FAssetDataToken::Create(Error.ReferencedAsset));
+
+			if (Error.Type == EAssetReferenceErrorType::Illegal)
+			{
+				if (!AssetReferenceFixer)
+				{
+					AssetReferenceFixer = GEditor->MakeAssetReferenceFixer();
+				}
+				if (AssetReferenceFixer)
+				{
+					if (TSharedPtr<UE::DataValidation::IFixer> Fixer = AssetReferenceFixer->CreateFixer(InAssetData))
+					{
+						TokenizedMessage->AddToken(Fixer->CreateToken(AssetReferenceFixer->GetFixerLabel(InAssetData)));
+					}
+				}
+			}
 		}
 	}
 }
diff --git a/Engine/Source/Editor/UnrealEd/Classes/Editor/AssetReferenceFixer.h b/Engine/Source/Editor/UnrealEd/Classes/Editor/AssetReferenceFixer.h
new file mode 100644
--- /dev/null
+++ b/Engine/Source/Editor/UnrealEd/Classes/Editor/AssetReferenceFixer.h
@@ -0,0 +1,38 @@
+// Copyright Epic Games, Inc. All Rights Reserved.
+
+#pragma once
+
+#include "AssetRegistry/AssetData.h"
+#include "Internationalization/Text.h"
+
+namespace UE::DataValidation
+{
+	struct IFixer;
+}
+
+/** Used in fixing invalid references between assets. Implement a subclass of this and return it in OnMakeAssetReferenceFixer */
+class IAssetReferenceFixer
+{
+public:
+	virtual ~IAssetReferenceFixer() = default;
+
+	/** Create a reference fixer for the given an asset, or null if no automatic fix-up can be performed for the given asset. */
+	virtual TSharedPtr<UE::DataValidation::IFixer> CreateFixer(const FAssetData& AssetData) const = 0;
+
+	/** Get the label to use as the fixer for the given asset. */
+	FText GetFixerLabel(const FAssetData& AssetData) const
+	{
+		FText FixerLabel = GetFixerLabelImpl(AssetData);
+		if (FixerLabel.IsEmpty())
+		{
+			FixerLabel = FText::FromName(AssetData.AssetName);
+		}
+		return FixerLabel;
+	}
+
+protected:
+	virtual FText GetFixerLabelImpl(const FAssetData& AssetData) const
+	{
+		return FText();
+	}
+};
diff --git a/Engine/Source/Editor/UnrealEd/Classes/Editor/EditorEngine.h b/Engine/Source/Editor/UnrealEd/Classes/Editor/EditorEngine.h
--- a/Engine/Source/Editor/UnrealEd/Classes/Editor/EditorEngine.h
+++ b/Engine/Source/Editor/UnrealEd/Classes/Editor/EditorEngine.h
@@ -56,6 +56,7 @@ class FViewport;
 class IEngineLoop;
 class ILauncherWorker;
 class ILayers;
+class IAssetReferenceFixer;
 class IAssetViewport;
 class ITargetPlatform;
 class SViewport;
@@ -2348,6 +2349,11 @@ public:
 	FOnMakeAssetReferenceFilter& OnMakeAssetReferenceFilter() { return OnMakeAssetReferenceFilterDelegate; }
 	TSharedPtr<IAssetReferenceFilter> MakeAssetReferenceFilter(const FAssetReferenceFilterContext& Context) { return OnMakeAssetReferenceFilterDelegate.IsBound() ? OnMakeAssetReferenceFilterDelegate.Execute(Context) : nullptr; }
 
+	/** Returns a fixer to resolve illegal references between assets */
+	DECLARE_DELEGATE_RetVal(TSharedPtr<IAssetReferenceFixer>, FOnMakeAssetReferenceFixer);
+	FOnMakeAssetReferenceFixer& OnMakeAssetReferenceFixer() { return OnMakeAssetReferenceFixerDelegate; }
+	TSharedPtr<IAssetReferenceFixer> MakeAssetReferenceFixer() { return OnMakeAssetReferenceFixerDelegate.IsBound() ? OnMakeAssetReferenceFixerDelegate.Execute() : nullptr; }
+
 	DECLARE_DELEGATE_RetVal(ULevelEditorDragDropHandler*, FOnCreateLevelEditorDragDropHandler);
 	FOnCreateLevelEditorDragDropHandler& OnCreateLevelEditorDragDropHandler() { return OnCreateLevelEditorDragDropHandlerDelegate; }
 	UNREALED_API ULevelEditorDragDropHandler* GetLevelEditorDragDropHandler() const;
@@ -2371,6 +2377,7 @@ private:
 	mutable TObjectPtr<ULevelEditorDragDropHandler> DragDropHandler;
 
 	FOnMakeAssetReferenceFilter OnMakeAssetReferenceFilterDelegate;
+	FOnMakeAssetReferenceFixer OnMakeAssetReferenceFixerDelegate;
 	FOnCreateLevelEditorDragDropHandler OnCreateLevelEditorDragDropHandlerDelegate;
 	FOnFilterCopiedActors OnFilterCopiedActorsDelegate;
 
