Added to FDomainAssetReferenceFilter a delegate implemented by world partition to handle special cases of allowed cross-plugin references (plugin actor part of an Externa lData Layer is allowed to reference its world part of another plugin).

diff --git a/Engine/Plugins/Editor/AssetReferenceRestrictions/Source/AssetReferenceRestrictions/Private/AssetReferencingDomains.cpp b/Engine/Plugins/Editor/AssetReferenceRestrictions/Source/AssetReferenceRestrictions/Private/AssetReferencingDomains.cpp
--- a/Engine/Plugins/Editor/AssetReferenceRestrictions/Source/AssetReferenceRestrictions/Private/AssetReferencingDomains.cpp
+++ b/Engine/Plugins/Editor/AssetReferenceRestrictions/Source/AssetReferenceRestrictions/Private/AssetReferencingDomains.cpp
@@ -170,6 +170,8 @@ void FDomainDatabase::RebuildFromScratch()
 	GameDomain.Reset();
 	NeverCookDomain.Reset();
 	DomainsDefinedByPlugins.Reset();
+	PluginDomains.Reset();
+
 	SpecificAssetPackageDomains.Reset();
 
 	// Create the built-in domains
@@ -242,7 +244,6 @@ void FDomainDatabase::RebuildFromScratch()
 	}
 
 	// Create the domains for plugins that contain content
-	DomainsDefinedByPlugins.Reset();
 	TArray<TSharedRef<IPlugin>> EnabledPlugins = IPluginManager::Get().GetEnabledPlugins();
 	for (const TSharedRef<IPlugin>& Plugin : EnabledPlugins)
 	{
@@ -373,6 +374,7 @@ void FDomainDatabase::BuildDomainFromPlugin(TSharedRef<IPlugin> Plugin)
 	DomainsDefinedByPlugins.Add(NewDomainName);
 
 	TSharedPtr<FDomainData> Domain = FindOrAddDomainByName(NewDomainName);
+	PluginDomains.Add(Domain);
 
 	Domain->Reset();
 	Domain->DomainRootPaths.Add(Plugin->GetMountedAssetPath());
@@ -464,6 +466,7 @@ void FDomainDatabase::BuildUnrestrictedDomainFromPlugin(TSharedRef<IPlugin> Plug
 	DomainsDefinedByPlugins.Add(NewDomainName);
 
 	TSharedPtr<FDomainData> Domain = FindOrAddDomainByName(NewDomainName);
+	PluginDomains.Add(Domain);
 
 	Domain->Reset();
 	Domain->DomainRootPaths.Add(Plugin->GetMountedAssetPath());
diff --git a/Engine/Plugins/Editor/AssetReferenceRestrictions/Source/AssetReferenceRestrictions/Private/DomainAssetReferenceFilter.cpp b/Engine/Plugins/Editor/AssetReferenceRestrictions/Source/AssetReferenceRestrictions/Private/DomainAssetReferenceFilter.cpp
--- a/Engine/Plugins/Editor/AssetReferenceRestrictions/Source/AssetReferenceRestrictions/Private/DomainAssetReferenceFilter.cpp
+++ b/Engine/Plugins/Editor/AssetReferenceRestrictions/Source/AssetReferenceRestrictions/Private/DomainAssetReferenceFilter.cpp
@@ -34,8 +34,8 @@ void FDomainAssetReferenceFilter::DetermineReferencingDomain()
 {
 	QUICK_SCOPE_CYCLE_COUNTER(FDomainAssetReferenceFilter_DetermineReferencingDomain);
 
-	ReferencingDomains.Reset();
-
+	ReferencingAssetDataInfos.Reset();
+	
 	FAssetRegistryModule& AssetRegistryModule = FModuleManager::LoadModuleChecked<FAssetRegistryModule>("AssetRegistry");
 
 	// Collect the associated domains from each referencing asset
@@ -82,7 +82,7 @@ void FDomainAssetReferenceFilter::DetermineReferencingDomain()
 		TSharedPtr<FDomainData> TestDomain = DomainDB->FindDomainFromAssetData(TestAsset);
 		if (TestDomain.IsValid())
 		{
-			ReferencingDomains.Add(TestDomain);
+			ReferencingAssetDataInfos.Add(MakeTuple(TestAsset, TestDomain));
 		}
 	}
 }
@@ -107,6 +107,18 @@ void FDomainAssetReferenceFilter::UpdateAllFilters()
 	}
 }
 
+bool FDomainAssetReferenceFilter::IsCrossPluginReferenceAllowed(const FAssetDataInfo& ReferencingAssetDataInfo, const FAssetDataInfo& ReferencedAssetDataInfo) const
+{
+	const TSharedPtr<FDomainData>& ReferencingAssetDomain = ReferencingAssetDataInfo.Value;
+	const TSharedPtr<FDomainData>& ReferencedAssetDomain = ReferencedAssetDataInfo.Value;
+	const FAssetData& ReferencingAssetData = ReferencingAssetDataInfo.Key;
+	const FAssetData& ReferencedAssetData = ReferencedAssetDataInfo.Key;
+
+	return DomainDB->IsPlugin(ReferencedAssetDomain) &&
+		   DomainDB->IsPlugin(ReferencingAssetDomain) &&
+		   IAssetReferenceFilter::IsCrossPluginReferenceAllowed(ReferencingAssetData, ReferencedAssetData);
+}
+
 bool FDomainAssetReferenceFilter::PassesFilterImpl(const FAssetData& AssetData, FText& OutOptionalFailureReason) const
 {
 	QUICK_SCOPE_CYCLE_COUNTER(FDomainAssetReferenceFilter_PassesFilterImpl);
@@ -126,9 +138,12 @@ bool FDomainAssetReferenceFilter::PassesFilterImpl(const FAssetData& AssetData,
 	// If it gets hit, transient can be added to the special system mount 'Temp' domain
 	if (ensureMsgf(AssetDomain.IsValid(), TEXT("Asset %s didn't match any domain"), *AssetData.GetObjectPathString()))
 	{
+		FAssetDataInfo ReferencedAssetDataInfo = MakeTuple(AssetData, AssetDomain);
+
 		// Check the referencing domains
-		for (const TSharedPtr<FDomainData>& ReferencingDomain : ReferencingDomains)
+		for (const FAssetDataInfo& ReferencingAssetDataInfo : ReferencingAssetDataInfos)
 		{
+			const TSharedPtr<FDomainData>& ReferencingDomain = ReferencingAssetDataInfo.Value;
 			if (ReferencingDomain->bCanSeeEverything)
 			{
 				return true;
@@ -139,6 +154,12 @@ bool FDomainAssetReferenceFilter::PassesFilterImpl(const FAssetData& AssetData,
 
 			if (!bResult)
 			{
+				// Before returning false, check whether this is an allowed cross-plugin asset reference (exception)
+				if (IsCrossPluginReferenceAllowed(ReferencingAssetDataInfo, ReferencedAssetDataInfo))
+				{
+					return true;
+				}
+				
 				return false;
 			}
 		}
diff --git a/Engine/Plugins/Editor/AssetReferenceRestrictions/Source/AssetReferenceRestrictions/Public/AssetReferencingDomains.h b/Engine/Plugins/Editor/AssetReferenceRestrictions/Source/AssetReferenceRestrictions/Public/AssetReferencingDomains.h
--- a/Engine/Plugins/Editor/AssetReferenceRestrictions/Source/AssetReferenceRestrictions/Public/AssetReferencingDomains.h
+++ b/Engine/Plugins/Editor/AssetReferenceRestrictions/Source/AssetReferenceRestrictions/Public/AssetReferencingDomains.h
@@ -74,6 +74,7 @@ struct ASSETREFERENCERESTRICTIONS_API FDomainDatabase final
 	TTuple<bool, FText> CanDomainsSeeEachOther(TSharedPtr<FDomainData> Referencee, TSharedPtr<FDomainData> Referencer) const;
 
 	const TArray<FString>& GetDomainsDefinedByPlugins() const { return DomainsDefinedByPlugins; }
+	bool IsPlugin(TSharedPtr<FDomainData> Domain) const { return PluginDomains.Contains(Domain); }
 
 	TSharedPtr<FDomainData> FindOrAddDomainByName(const FString& Name);
 
@@ -113,6 +114,9 @@ private:
 	// List of domains that came from plugins (used for domain pickers in the settings)
 	TArray<FString> DomainsDefinedByPlugins;
 
+	// Plugin Domains
+	TSet<TSharedPtr<FDomainData>> PluginDomains;
+
 	bool bDatabaseOutOfDate = false;
 
 	TMulticastDelegate<void()> OnPostDatabaseUpdated;
diff --git a/Engine/Plugins/Editor/AssetReferenceRestrictions/Source/AssetReferenceRestrictions/Public/DomainAssetReferenceFilter.h b/Engine/Plugins/Editor/AssetReferenceRestrictions/Source/AssetReferenceRestrictions/Public/DomainAssetReferenceFilter.h
--- a/Engine/Plugins/Editor/AssetReferenceRestrictions/Source/AssetReferenceRestrictions/Public/DomainAssetReferenceFilter.h
+++ b/Engine/Plugins/Editor/AssetReferenceRestrictions/Source/AssetReferenceRestrictions/Public/DomainAssetReferenceFilter.h
@@ -23,7 +23,10 @@ public:
 	// Update any cached information for all filters
 	static void UpdateAllFilters();
 private:
+	using FAssetDataInfo = TTuple<FAssetData, TSharedPtr<FDomainData>>;
+
 	bool PassesFilterImpl(const FAssetData& AssetData, FText& OutOptionalFailureReason) const;
+	bool IsCrossPluginReferenceAllowed(const FAssetDataInfo& ReferencingAssetDataInfo, const FAssetDataInfo& ReferencedAssetDataInfo) const;
 
 	void DetermineReferencingDomain();
 
@@ -36,7 +39,7 @@ private:
 	TSharedPtr<FDomainDatabase> DomainDB;
 
 	TArray<FAssetReferenceFilterReferencerInfo> OriginalReferencingAssets;
-	TSet<TSharedPtr<FDomainData>> ReferencingDomains;
+	TSet<FAssetDataInfo> ReferencingAssetDataInfos;
 
 	FText Failure_CouldNotDetermineDomain;
 };
diff --git a/Engine/Source/Editor/UnrealEd/Classes/Editor/AssetReferenceFilter.h b/Engine/Source/Editor/UnrealEd/Classes/Editor/AssetReferenceFilter.h
--- a/Engine/Source/Editor/UnrealEd/Classes/Editor/AssetReferenceFilter.h
+++ b/Engine/Source/Editor/UnrealEd/Classes/Editor/AssetReferenceFilter.h
@@ -3,6 +3,7 @@
 #pragma once
 
 #include "CoreTypes.h"
+#include "Delegates/Delegate.h"
 
 class FText;
 struct FAssetData;
@@ -14,4 +15,15 @@ public:
 	virtual ~IAssetReferenceFilter() { }
 	/** Filter function to pass/fail an asset. Called in some situations that are performance-sensitive so is expected to be fast. */
 	virtual bool PassesFilter(const FAssetData& AssetData, FText* OutOptionalFailureReason = nullptr) const = 0;
+
+	DECLARE_DELEGATE_RetVal_TwoParams(bool, FOnIsCrossPluginReferenceAllowed, const FAssetData& /*ReferencingAssetData*/, const FAssetData& /*ReferencedAssetData*/);
+
+	/** Used to register a custom delegate implementation that returns if a plugin asset can reference another plugin asset. */
+	UNREALED_API static FOnIsCrossPluginReferenceAllowed& OnIsCrossPluginReferenceAllowed() { return OnIsCrossPluginReferenceAllowedDelegate; }
+	
+protected:
+	UNREALED_API static bool IsCrossPluginReferenceAllowed(const FAssetData& ReferencingAssetData, const FAssetData& ReferencedAssetData);
+
+private:
+	static FOnIsCrossPluginReferenceAllowed OnIsCrossPluginReferenceAllowedDelegate;
 };
\ No newline at end of file
diff --git a/Engine/Source/Editor/UnrealEd/Private/AssetReferenceFilter.cpp b/Engine/Source/Editor/UnrealEd/Private/AssetReferenceFilter.cpp
new file mode 100644
--- /dev/null
+++ b/Engine/Source/Editor/UnrealEd/Private/AssetReferenceFilter.cpp
@@ -0,0 +1,11 @@
+// Copyright Epic Games, Inc. All Rights Reserved.
+
+#include "Editor/AssetReferenceFilter.h"
+#include "AssetRegistry/AssetData.h"
+
+IAssetReferenceFilter::FOnIsCrossPluginReferenceAllowed IAssetReferenceFilter::OnIsCrossPluginReferenceAllowedDelegate;
+
+bool IAssetReferenceFilter::IsCrossPluginReferenceAllowed(const FAssetData& ReferencingAssetData, const FAssetData& ReferencedAssetData)
+{
+	return OnIsCrossPluginReferenceAllowedDelegate.IsBound() ? OnIsCrossPluginReferenceAllowedDelegate.Execute(ReferencingAssetData, ReferencedAssetData) : false;
+}
\ No newline at end of file
diff --git a/Engine/Source/Editor/WorldPartitionEditor/Private/WorldPartitionEditorModule.cpp b/Engine/Source/Editor/WorldPartitionEditor/Private/WorldPartitionEditorModule.cpp
--- a/Engine/Source/Editor/WorldPartitionEditor/Private/WorldPartitionEditorModule.cpp
+++ b/Engine/Source/Editor/WorldPartitionEditor/Private/WorldPartitionEditorModule.cpp
@@ -22,12 +22,15 @@
 #include "WorldPartition/HLOD/SWorldPartitionBuildHLODsDialog.h"
 #include "WorldPartition/HLOD/HLODEditorSubsystem.h"
 #include "WorldPartition/WorldPartitionClassDescRegistry.h"
+#include "WorldPartition/DataLayer/ExternalDataLayerUID.h"
+#include "WorldPartition/DataLayer/ExternalDataLayerHelper.h"
 
 #include "LevelEditor.h"
 #include "LevelEditorViewport.h"
 #include "Framework/Application/SlateApplication.h"
 #include "Framework/MultiBox/MultiBoxBuilder.h"
 
+#include "Editor/AssetReferenceFilter.h"
 #include "Engine/Level.h"
 
 #include "Misc/FileHelper.h"
@@ -248,6 +251,8 @@ void FWorldPartitionEditorModule::StartupModule()
 			}
 		}
 	});
+
+	IAssetReferenceFilter::OnIsCrossPluginReferenceAllowed().BindRaw(this, &FWorldPartitionEditorModule::OnIsCrossPluginReferenceAllowed);
 }
 
 void FWorldPartitionEditorModule::ShutdownModule()
@@ -255,6 +260,8 @@ void FWorldPartitionEditorModule::ShutdownModule()
 	FWorldPartitionClassDescRegistry().Get().Uninitialize();
 	FWorldPartitionClassDescRegistry().Get().TearDown();
 
+	IAssetReferenceFilter::OnIsCrossPluginReferenceAllowed().Unbind();
+
 	if (!IsRunningGame())
 	{
 		if (FLevelEditorModule* LevelEditorModule = FModuleManager::Get().GetModulePtr<FLevelEditorModule>("LevelEditor"))
@@ -284,6 +291,25 @@ void FWorldPartitionEditorModule::ShutdownModule()
 	FEditorDelegates::OnEditorInitialized.Remove(EditorInitializedHandle);
 }
 
+bool FWorldPartitionEditorModule::OnIsCrossPluginReferenceAllowed(const FAssetData& ReferencingAssetData, const FAssetData& ReferencedAssetData)
+{
+	// Allow External Data Layer Actor (ReferencingAssetData) from a plugin X to reference its world (ReferencedAssetData) from a plugin Y
+	const UClass* ReferencedAssetDataClass = ReferencedAssetData.GetClass();
+	if (ReferencedAssetDataClass && ReferencedAssetDataClass->IsChildOf<UWorld>())
+	{
+		const FString ReferencingAssetPath = ReferencingAssetData.PackagePath.ToString();
+		FExternalDataLayerUID ReferencingExternalDataLayerUID;
+		if (FExternalDataLayerHelper::IsExternalDataLayerPath(ReferencingAssetPath, &ReferencingExternalDataLayerUID); ReferencingExternalDataLayerUID.IsValid())
+		{
+			// Use referencing asset's optional outer path name (if any) to build its package name and compare it with the referenced package name
+			const FString ReferencingOptionalOuterPackageName = FSoftObjectPath(ReferencingAssetData.GetOptionalOuterPathName().ToString()).GetLongPackageName();
+			return ReferencingOptionalOuterPackageName == ReferencedAssetData.PackageName.ToString();
+		}
+	}
+
+	return false;
+}
+
 void FWorldPartitionEditorModule::RegisterMenus()
 {
 	FLevelEditorModule& LevelEditorModule = FModuleManager::Get().LoadModuleChecked<FLevelEditorModule>("LevelEditor");
diff --git a/Engine/Source/Editor/WorldPartitionEditor/Public/WorldPartitionEditorModule.h b/Engine/Source/Editor/WorldPartitionEditor/Public/WorldPartitionEditorModule.h
--- a/Engine/Source/Editor/WorldPartitionEditor/Public/WorldPartitionEditorModule.h
+++ b/Engine/Source/Editor/WorldPartitionEditor/Public/WorldPartitionEditorModule.h
@@ -9,6 +9,7 @@ class FLayoutExtender;
 class SDockTab;
 class FSpawnTabArgs;
 class SContentBundleBrowser;
+struct FAssetData;
 
 
 /**
@@ -158,6 +159,7 @@ private:
 private:
 	void RunCommandletAsExternalProcess(const FString& InCommandletArgs, const FText& InOperationDescription, int32& OutResult, bool& bOutCancelled);
 	void OnConvertMap();
+	bool OnIsCrossPluginReferenceAllowed(const FAssetData& ReferencingAssetData, const FAssetData& ReferencedAssetData);
 
 	FDelegateHandle EditorInitializedHandle;
 	FDelegateHandle LevelEditorExtenderDelegateHandle;
