Include the validator class in its log messages

This matches other validators, and makes filtering/inspecting the log output easier

#rnx

diff --git a/Engine/Plugins/Editor/AssetReferenceRestrictions/Source/AssetReferenceRestrictions/Private/AssetValidator_AssetReferenceRestrictions.cpp b/Engine/Plugins/Editor/AssetReferenceRestrictions/Source/AssetReferenceRestrictions/Private/AssetValidator_AssetReferenceRestrictions.cpp
--- a/Engine/Plugins/Editor/AssetReferenceRestrictions/Source/AssetReferenceRestrictions/Private/AssetValidator_AssetReferenceRestrictions.cpp
+++ b/Engine/Plugins/Editor/AssetReferenceRestrictions/Source/AssetReferenceRestrictions/Private/AssetValidator_AssetReferenceRestrictions.cpp
@@ -68,7 +68,9 @@ void UAssetValidator_AssetReferenceRestrictions::ValidateAssetInternal(const FAs
 		TSharedPtr<IAssetReferenceFixer> AssetReferenceFixer;
 		for (const FAssetReferenceError& Error : Result.GetError())
 		{
-			TSharedRef<FTokenizedMessage> TokenizedMessage = AssetMessage(InAssetData, EMessageSeverity::Error, Error.Message)->AddToken(FAssetDataToken::Create(Error.ReferencedAsset));
+			TSharedRef<FTokenizedMessage> TokenizedMessage = AssetMessage(InAssetData, EMessageSeverity::Error, Error.Message)
+				->AddToken(FAssetDataToken::Create(Error.ReferencedAsset))
+				->AddToken(FTextToken::Create(FText::FormatNamed(LOCTEXT("ValidatorClassSuffix", ". ({ValidatorName})"), TEXT("ValidatorName"), FText::AsCultureInvariant(GetClass()->GetName()))));
 
 			if (Error.Type == EAssetReferenceErrorType::Illegal)
 			{
