Asset reference filter improvements

* IAssetReferenceFilter can now tell whether the reference request is for validation (eg, from UAssetValidator_AssetReferenceRestrictions) via the EAssetReferenceFilterRole option given to its FAssetReferenceFilterContext.
* IAssetReferenceFilter can now downgrade illegal reference errors for a given asset to warnings, by overriding DoesAssetDowngradeReferenceErrorsToWarnings and having it return true.

diff --git a/Engine/Plugins/Editor/AssetReferenceRestrictions/Source/AssetReferenceRestrictions/Private/AssetReferencingPolicySubsystem.cpp b/Engine/Plugins/Editor/AssetReferenceRestrictions/Source/AssetReferenceRestrictions/Private/AssetReferencingPolicySubsystem.cpp
--- a/Engine/Plugins/Editor/AssetReferenceRestrictions/Source/AssetReferenceRestrictions/Private/AssetReferencingPolicySubsystem.cpp
+++ b/Engine/Plugins/Editor/AssetReferenceRestrictions/Source/AssetReferenceRestrictions/Private/AssetReferencingPolicySubsystem.cpp
@@ -45,6 +45,11 @@ bool UAssetReferencingPolicySubsystem::ShouldValidateAssetReferences(const FAsse
 }
 
 TValueOrError<void, TArray<FAssetReferenceError>> UAssetReferencingPolicySubsystem::ValidateAssetReferences(const FAssetData& InAssetData) const
+{
+	return ValidateAssetReferences(InAssetData, EAssetReferenceFilterRole::None);
+}
+
+TValueOrError<void, TArray<FAssetReferenceError>> UAssetReferencingPolicySubsystem::ValidateAssetReferences(const FAssetData& InAssetData, const EAssetReferenceFilterRole Role) const
 {
 	IAssetRegistry& AssetRegistry = FModuleManager::LoadModuleChecked<FAssetRegistryModule>("AssetRegistry").Get();
 	TArray<FAssetReferenceError> Errors;
@@ -101,16 +106,19 @@ TValueOrError<void, TArray<FAssetReferenceError>> UAssetReferencingPolicySubsyst
 	if (AllDependencyAssets.Num() > 0)
 	{
 		FAssetReferenceFilterContext AssetReferenceFilterContext;
+		AssetReferenceFilterContext.AddRole(Role);
 		AssetReferenceFilterContext.AddReferencingAsset(InAssetData);
 		TSharedPtr<IAssetReferenceFilter> AssetReferenceFilter = GEditor ? GEditor->MakeAssetReferenceFilter(AssetReferenceFilterContext) : nullptr;
 		if (ensure(AssetReferenceFilter.IsValid()))
 		{
+			const bool bReferenceErrorsAsWarnings = AssetReferenceFilter->DoesAssetDowngradeReferenceErrorsToWarnings(InAssetData);
 			for (const FAssetData& Dependency : AllDependencyAssets)
 			{
 				FText FailureReason;
 				if (!AssetReferenceFilter->PassesFilter(Dependency, &FailureReason))
 				{
 					Errors.Add(FAssetReferenceError{
+						.bTreatErrorAsWarning = bReferenceErrorsAsWarnings,
 						.Type = EAssetReferenceErrorType::Illegal,
 						.ReferencedAsset = Dependency,
 						.Message = FText::Format(LOCTEXT("IllegalReference_AssetFilterFail", "Illegal reference: {0}"), FailureReason),
diff --git a/Engine/Plugins/Editor/AssetReferenceRestrictions/Source/AssetReferenceRestrictions/Private/AssetValidator_AssetReferenceRestrictions.cpp b/Engine/Plugins/Editor/AssetReferenceRestrictions/Source/AssetReferenceRestrictions/Private/AssetValidator_AssetReferenceRestrictions.cpp
--- a/Engine/Plugins/Editor/AssetReferenceRestrictions/Source/AssetReferenceRestrictions/Private/AssetValidator_AssetReferenceRestrictions.cpp
+++ b/Engine/Plugins/Editor/AssetReferenceRestrictions/Source/AssetReferenceRestrictions/Private/AssetValidator_AssetReferenceRestrictions.cpp
@@ -62,13 +62,13 @@ EDataValidationResult UAssetValidator_AssetReferenceRestrictions::ValidateLoaded
 void UAssetValidator_AssetReferenceRestrictions::ValidateAssetInternal(const FAssetData& InAssetData, const IAssetRegistry& InAssetRegistry)
 {
     UAssetReferencingPolicySubsystem* Subsystem = GEditor->GetEditorSubsystem<UAssetReferencingPolicySubsystem>();
-	if (TValueOrError<void, TArray<FAssetReferenceError>> Result = Subsystem->ValidateAssetReferences(InAssetData);
+	if (TValueOrError<void, TArray<FAssetReferenceError>> Result = Subsystem->ValidateAssetReferences(InAssetData, EAssetReferenceFilterRole::Validation);
 		Result.HasError())
 	{
 		TSharedPtr<IAssetReferenceFixer> AssetReferenceFixer;
 		for (const FAssetReferenceError& Error : Result.GetError())
 		{
-			TSharedRef<FTokenizedMessage> TokenizedMessage = AssetMessage(InAssetData, EMessageSeverity::Error, Error.Message)
+			TSharedRef<FTokenizedMessage> TokenizedMessage = AssetMessage(InAssetData, Error.bTreatErrorAsWarning ? EMessageSeverity::Warning : EMessageSeverity::Error, Error.Message)
 				->AddToken(FAssetDataToken::Create(Error.ReferencedAsset))
 				->AddToken(FTextToken::Create(FText::FormatNamed(LOCTEXT("ValidatorClassSuffix", ". ({ValidatorName})"), TEXT("ValidatorName"), FText::AsCultureInvariant(GetClass()->GetName()))));
 
diff --git a/Engine/Plugins/Editor/AssetReferenceRestrictions/Source/AssetReferenceRestrictions/Public/AssetReferencingPolicySubsystem.h b/Engine/Plugins/Editor/AssetReferenceRestrictions/Source/AssetReferenceRestrictions/Public/AssetReferencingPolicySubsystem.h
--- a/Engine/Plugins/Editor/AssetReferenceRestrictions/Source/AssetReferenceRestrictions/Public/AssetReferencingPolicySubsystem.h
+++ b/Engine/Plugins/Editor/AssetReferenceRestrictions/Source/AssetReferenceRestrictions/Public/AssetReferencingPolicySubsystem.h
@@ -8,6 +8,7 @@
 
 #include "AssetReferencingPolicySubsystem.generated.h"
 
+enum class EAssetReferenceFilterRole : uint8;
 struct FAssetReferenceFilterContext;
 class IAssetReferenceFilter;
 struct FAssetData;
@@ -21,6 +22,7 @@ enum class EAssetReferenceErrorType
 
 struct FAssetReferenceError
 {
+	bool bTreatErrorAsWarning = false;
 	EAssetReferenceErrorType Type;
 	FAssetData ReferencedAsset;
 	FText Message;
@@ -44,6 +46,7 @@ public:
 
 	// Check the outgoing references of the given asset according to the asset registry and return details of any errors 
 	TValueOrError<void, TArray<FAssetReferenceError>> ValidateAssetReferences(const FAssetData& Asset) const;
+	TValueOrError<void, TArray<FAssetReferenceError>> ValidateAssetReferences(const FAssetData& Asset, const EAssetReferenceFilterRole Role) const;
 
 	TSharedPtr<FDomainDatabase> GetDomainDB() const;
 private:
diff --git a/Engine/Source/Editor/UnrealEd/Classes/Editor/AssetReferenceFilter.h b/Engine/Source/Editor/UnrealEd/Classes/Editor/AssetReferenceFilter.h
--- a/Engine/Source/Editor/UnrealEd/Classes/Editor/AssetReferenceFilter.h
+++ b/Engine/Source/Editor/UnrealEd/Classes/Editor/AssetReferenceFilter.h
@@ -16,6 +16,12 @@ public:
 	/** Filter function to pass/fail an asset. Called in some situations that are performance-sensitive so is expected to be fast. */
 	virtual bool PassesFilter(const FAssetData& AssetData, FText* OutOptionalFailureReason = nullptr) const = 0;
 
+	/** Return true if the given asset downgrades any illegal reference errors to warnings */
+	virtual bool DoesAssetDowngradeReferenceErrorsToWarnings(const FAssetData& AssetData) const
+	{
+		return false;
+	}
+
 	DECLARE_DELEGATE_RetVal_TwoParams(bool, FOnIsCrossPluginReferenceAllowed, const FAssetData& /*ReferencingAssetData*/, const FAssetData& /*ReferencedAssetData*/);
 
 	/** Used to register a custom delegate implementation that returns if a plugin asset can reference another plugin asset. */
diff --git a/Engine/Source/Editor/UnrealEd/Classes/Editor/EditorEngine.h b/Engine/Source/Editor/UnrealEd/Classes/Editor/EditorEngine.h
--- a/Engine/Source/Editor/UnrealEd/Classes/Editor/EditorEngine.h
+++ b/Engine/Source/Editor/UnrealEd/Classes/Editor/EditorEngine.h
@@ -309,6 +309,15 @@ struct FAssetReferenceFilterReferencerInfo
 	}
 };
 
+enum class EAssetReferenceFilterRole : uint8
+{
+	/** This filter has no specific role */
+	None = 0,
+	/** This filter is being used by illegal reference validation */
+	Validation = 1<<0,
+};
+ENUM_CLASS_FLAGS(EAssetReferenceFilterRole);
+
 /** Struct used in filtering allowed references between assets. Passes context about the referencers to game-level filters */
 struct FAssetReferenceFilterContext
 {
@@ -321,8 +330,29 @@ struct FAssetReferenceFilterContext
 	UE_DEPRECATED(5.5, "ReferencingAssets is deprecated. Use the AddReferencingAsset functions to add referencers and GetReferencingAssets to get them.")
 	TArray<FAssetData> ReferencingAssets;
 
+	void AddRole(const EAssetReferenceFilterRole InRole)
+	{
+		EnumAddFlags(Role, InRole);
+	}
+
+	void RemoveRole(const EAssetReferenceFilterRole InRole)
+	{
+		EnumRemoveFlags(Role, InRole);
+	}
+
+	bool HasRole(const EAssetReferenceFilterRole InRole) const
+	{
+		return EnumHasAnyFlags(Role, InRole);
+	}
+
+	EAssetReferenceFilterRole GetRole() const
+	{
+		return Role;
+	}
+
 private:
 	TArray<FAssetReferenceFilterReferencerInfo> ReferencingAssetInfo;
+	EAssetReferenceFilterRole Role = EAssetReferenceFilterRole::None;
 };
 
 /**
