Check for missing hard refs to external actors during asset reference validation

#rnx

diff --git a/Engine/Plugins/Editor/AssetReferenceRestrictions/Source/AssetReferenceRestrictions/Private/AssetReferencingPolicySubsystem.cpp b/Engine/Plugins/Editor/AssetReferenceRestrictions/Source/AssetReferenceRestrictions/Private/AssetReferencingPolicySubsystem.cpp
--- a/Engine/Plugins/Editor/AssetReferenceRestrictions/Source/AssetReferenceRestrictions/Private/AssetReferencingPolicySubsystem.cpp
+++ b/Engine/Plugins/Editor/AssetReferenceRestrictions/Source/AssetReferenceRestrictions/Private/AssetReferencingPolicySubsystem.cpp
@@ -13,6 +13,16 @@
 
 #define LOCTEXT_NAMESPACE "AssetReferencingPolicy"
 
+namespace AssetReferencingPolicySubsystem::Private
+{
+	bool bCheckForMissingHardRefsToExternalActors = false;
+	static FAutoConsoleVariableRef CVarCheckForMissingHardRefsToExternalActors(
+		TEXT("AssetReferencingPolicy.CheckForMissingHardRefsToExternalActors"),
+		bCheckForMissingHardRefsToExternalActors,
+		TEXT("Error on missing hard dependencies to external actors when validating asset references.")
+	);
+}
+
 bool UAssetReferencingPolicySubsystem::ShouldCreateSubsystem(UObject* Outer) const
 {
 	return true;
@@ -100,7 +110,20 @@ TValueOrError<void, TArray<FAssetReferenceError>> UAssetReferencingPolicySubsyst
 	AssetRegistry.GetDependencies(PackageFName, HardDependencies, UE::AssetRegistry::EDependencyCategory::Package, UE::AssetRegistry::EDependencyQuery::Hard | QueryFlags);
 	for (FName HardDependency : HardDependencies)
 	{
-		AssetRegistry.GetAssetsByPackageName(HardDependency, AllDependencyAssets, true);
+		TArray<FAssetData> DependencyAssets;
+		AssetRegistry.GetAssetsByPackageName(HardDependency, DependencyAssets, true);
+		AllDependencyAssets.Append(DependencyAssets);
+		// Check for missing external actors
+		FNameBuilder HardDependencyStr(HardDependency);
+		const bool bIsExternalActorPackage = HardDependencyStr.ToView().Contains(FPackagePath::GetExternalActorsFolderName());
+		if (AssetReferencingPolicySubsystem::Private::bCheckForMissingHardRefsToExternalActors && bIsExternalActorPackage && DependencyAssets.Num() == 0)
+		{
+			Errors.Add(FAssetReferenceError{
+				.Type = EAssetReferenceErrorType::DoesNotExist,
+				.ReferencedAsset = FAssetData(HardDependency, FName(FPathViews::GetPath(HardDependencyStr)), FName(), FTopLevelAssetPath(UObject::StaticClass())),
+				.Message = FText::Format(LOCTEXT("IllegalReference_MissingHardRefToExternalActor", "Hard references external actor {0} which does not exist"), FText::FromName(HardDependency)),
+			});
+		}
 	}
 
 	if (AllDependencyAssets.Num() > 0)
