Make missing external actor ref validation independent of hard/soft/game/editoronly classification.

#rnx

diff --git a/Engine/Plugins/Editor/AssetReferenceRestrictions/Source/AssetReferenceRestrictions/Private/AssetReferencingPolicySubsystem.cpp b/Engine/Plugins/Editor/AssetReferenceRestrictions/Source/AssetReferenceRestrictions/Private/AssetReferencingPolicySubsystem.cpp
--- a/Engine/Plugins/Editor/AssetReferenceRestrictions/Source/AssetReferenceRestrictions/Private/AssetReferencingPolicySubsystem.cpp
+++ b/Engine/Plugins/Editor/AssetReferenceRestrictions/Source/AssetReferenceRestrictions/Private/AssetReferencingPolicySubsystem.cpp
@@ -15,11 +15,11 @@
 
 namespace AssetReferencingPolicySubsystem::Private
 {
-	bool bCheckForMissingHardRefsToExternalActors = false;
-	static FAutoConsoleVariableRef CVarCheckForMissingHardRefsToExternalActors(
-		TEXT("AssetReferencingPolicy.CheckForMissingHardRefsToExternalActors"),
-		bCheckForMissingHardRefsToExternalActors,
-		TEXT("Error on missing hard dependencies to external actors when validating asset references.")
+	bool bCheckForMissingRefsToExternalActors = false;
+	static FAutoConsoleVariableRef CVarCheckForMissingRefsToExternalActors(
+		TEXT("AssetReferencingPolicy.CheckForMissingRefsToExternalActors"),
+		bCheckForMissingRefsToExternalActors,
+		TEXT("Error on missing dependencies to external actors when validating asset references.")
 	);
 }
 
@@ -93,7 +93,7 @@ TValueOrError<void, TArray<FAssetReferenceError>> UAssetReferencingPolicySubsyst
 				{
 					Errors.Add(FAssetReferenceError{
 						.Type = EAssetReferenceErrorType::DoesNotExist,
-						.ReferencedAsset = FAssetData(SoftDependency, FName(FPathViews::GetPath(SoftDependencyStr)), FName(), FTopLevelAssetPath(UObject::StaticClass())),
+						.ReferencedAsset = FAssetData(SoftDependency, FName(FPathViews::GetPath(SoftDependencyStr)), FName(), UObject::StaticClass()->GetClassPathName()),
 						.Message = FText::Format(LOCTEXT("IllegalReference_MissingSoftRef", "Soft references {0} which does not exist"), FText::FromName(SoftDependency)),
 					});
 				}
@@ -110,19 +110,29 @@ TValueOrError<void, TArray<FAssetReferenceError>> UAssetReferencingPolicySubsyst
 	AssetRegistry.GetDependencies(PackageFName, HardDependencies, UE::AssetRegistry::EDependencyCategory::Package, UE::AssetRegistry::EDependencyQuery::Hard | QueryFlags);
 	for (FName HardDependency : HardDependencies)
 	{
-		TArray<FAssetData> DependencyAssets;
-		AssetRegistry.GetAssetsByPackageName(HardDependency, DependencyAssets, true);
-		AllDependencyAssets.Append(DependencyAssets);
-		// Check for missing external actors
-		FNameBuilder HardDependencyStr(HardDependency);
-		const bool bIsExternalActorPackage = HardDependencyStr.ToView().Contains(FPackagePath::GetExternalActorsFolderName());
-		if (AssetReferencingPolicySubsystem::Private::bCheckForMissingHardRefsToExternalActors && bIsExternalActorPackage && DependencyAssets.Num() == 0)
+		AssetRegistry.GetAssetsByPackageName(HardDependency, AllDependencyAssets, true);
+	}
+
+	// Check for missing external actor references which can manifest as LoadErrors and show as either Hard/Soft/Game/EditorOnly.
+	if (AssetReferencingPolicySubsystem::Private::bCheckForMissingRefsToExternalActors)
+	{
+		TArray<FName> PkgDependencies;
+		AssetRegistry.GetDependencies(PackageFName, PkgDependencies, UE::AssetRegistry::EDependencyCategory::Package);
+		for (FName Dependency : PkgDependencies)
 		{
-			Errors.Add(FAssetReferenceError{
-				.Type = EAssetReferenceErrorType::DoesNotExist,
-				.ReferencedAsset = FAssetData(HardDependency, FName(FPathViews::GetPath(HardDependencyStr)), FName(), FTopLevelAssetPath(UObject::StaticClass())),
-				.Message = FText::Format(LOCTEXT("IllegalReference_MissingHardRefToExternalActor", "Hard references external actor {0} which does not exist"), FText::FromName(HardDependency)),
-			});
+			TArray<FAssetData> DependencyAssets;
+			AssetRegistry.GetAssetsByPackageName(Dependency, DependencyAssets, true);
+			
+			FNameBuilder DependencyStr(Dependency);
+			const bool bIsExternalActorPackage = DependencyStr.ToView().Contains(FPackagePath::GetExternalActorsFolderName());
+			if (bIsExternalActorPackage && DependencyAssets.Num() == 0)
+			{
+				Errors.Add(FAssetReferenceError{
+					.Type = EAssetReferenceErrorType::DoesNotExist,
+					.ReferencedAsset = FAssetData(Dependency, FName(FPathViews::GetPath(DependencyStr)), FName(), UObject::StaticClass()->GetClassPathName()),
+					.Message = FText::Format(LOCTEXT("IllegalReference_MissingRefToExternalActor", "references external actor {0} which does not exist"), FText::FromName(Dependency)),
+				});
+			}
 		}
 	}
 
