When emitting errors for references from redirectors, add assets referencing the redirector to the message.

diff --git a/Engine/Plugins/Editor/AssetReferenceRestrictions/Source/AssetReferenceRestrictions/Private/EditorValidator_PluginAssetReferences.cpp b/Engine/Plugins/Editor/AssetReferenceRestrictions/Source/AssetReferenceRestrictions/Private/EditorValidator_PluginAssetReferences.cpp
--- a/Engine/Plugins/Editor/AssetReferenceRestrictions/Source/AssetReferenceRestrictions/Private/EditorValidator_PluginAssetReferences.cpp
+++ b/Engine/Plugins/Editor/AssetReferenceRestrictions/Source/AssetReferenceRestrictions/Private/EditorValidator_PluginAssetReferences.cpp
@@ -3,6 +3,7 @@
 #include "EditorValidator_PluginAssetReferences.h"
 
 #include "Algo/RemoveIf.h"
+#include "AssetReferencingPolicySettings.h"
 #include "AssetReferencingPolicySubsystem.h"
 #include "AssetRegistry/AssetData.h"
 #include "AssetRegistry/AssetDataToken.h"
@@ -12,10 +13,13 @@
 #include "EditorValidatorSubsystem.h"
 #include "Editor/EditorEngine.h"
 #include "Interfaces/IPluginManager.h"
+#include "Logging/StructuredLog.h"
 #include "Misc/PathViews.h"
 
 #include UE_INLINE_GENERATED_CPP_BY_NAME(EditorValidator_PluginAssetReferences)
 
+#define LOCTEXT_NAMESPACE "AssetReferencingPolicy"
+
 bool UEditorValidator_PluginAssetReferences::CanValidateAsset_Implementation(const FAssetData& InAssetData, UObject* InObject, FDataValidationContext& InContext) const
 {
     if (UDataValidationChangelist* Changelist = Cast<UDataValidationChangelist>(InObject))
@@ -49,6 +53,7 @@ EDataValidationResult UEditorValidator_PluginAssetReferences::ValidateLoadedAsse
             TSharedPtr<IPlugin> Plugin = PluginManager.FindPlugin(PluginName);
             if (Plugin.IsValid() && Plugin->IsEnabled())
             {
+                UE_LOGFMT(LogContentValidation, Display, "Validating asset references from plugin: {PluginName}", Plugin->GetName());
                 PluginsToValidate.Add(Plugin.ToSharedRef());
             }
         }
@@ -65,10 +70,23 @@ EDataValidationResult UEditorValidator_PluginAssetReferences::ValidateLoadedAsse
     UAssetReferencingPolicySubsystem* AssetReferencingPolicySubsystem = GEditor->GetEditorSubsystem<UAssetReferencingPolicySubsystem>();
 	UEditorValidatorSubsystem* EditorValidationSubsystem = GEditor->GetEditorSubsystem<UEditorValidatorSubsystem>();
 	FValidateAssetsSettings Settings;
-	FDataValidationContext ValidationContext(false, Settings.ValidationUsecase, {});
-    AssetsToValidate.SetNum(Algo::RemoveIf(AssetsToValidate, [AssetReferencingPolicySubsystem, EditorValidationSubsystem, &Settings, &ValidationContext](const FAssetData& Asset) {
-        return !AssetReferencingPolicySubsystem->ShouldValidateAssetReferences(Asset) || !EditorValidationSubsystem->ShouldValidateAsset(Asset, Settings, ValidationContext);
-    }));
+
+    auto FilterAssetsToValidate = [&Settings, AssetReferencingPolicySubsystem, EditorValidationSubsystem](TArray<FAssetData>& InAssets){
+        FDataValidationContext ValidationContext(false, Settings.ValidationUsecase, {});
+        InAssets.SetNum(Algo::RemoveIf(InAssets, [AssetReferencingPolicySubsystem, EditorValidationSubsystem, &Settings, &ValidationContext](const FAssetData& Asset) {
+            return !AssetReferencingPolicySubsystem->ShouldValidateAssetReferences(Asset) || !EditorValidationSubsystem->ShouldValidateAsset(Asset, Settings, ValidationContext);
+        }));
+    };
+
+    FilterAssetsToValidate(AssetsToValidate);
+
+	const UAssetReferencingPolicySettings* PolicySettings = GetDefault<UAssetReferencingPolicySettings>();
+	UE::AssetRegistry::EDependencyQuery QueryFlags = UE::AssetRegistry::EDependencyQuery::NoRequirements;
+	if (PolicySettings->bIgnoreEditorOnlyReferences)
+	{
+		// If the rules allow ignoring editor-only referencers, restrict the dependencies to game references
+		QueryFlags = UE::AssetRegistry::EDependencyQuery::Game;
+	}
 
     for (const FAssetData& Asset : AssetsToValidate)
     {
@@ -77,7 +95,32 @@ EDataValidationResult UEditorValidator_PluginAssetReferences::ValidateLoadedAsse
         {
             for (const FAssetReferenceError& Error : Result.GetError())
             {
-                AssetMessage(Asset, EMessageSeverity::Error, Error.Message)->AddToken(FAssetDataToken::Create(Error.ReferencedAsset));
+                TSharedRef<FTokenizedMessage> Msg = AssetMessage(Asset, EMessageSeverity::Error, Error.Message)
+                    ->AddToken(FAssetDataToken::Create(Error.ReferencedAsset));
+                if (Asset.IsRedirector())
+                {
+                    TArray<FAssetData> RedirectorReferencers;
+                    TArray<FName> RedirectorReferencerPackages;
+                    AssetRegistry.GetReferencers(Asset.PackageName, RedirectorReferencerPackages, UE::AssetRegistry::EDependencyCategory::Package, QueryFlags);
+                    for (FName PackageName : RedirectorReferencerPackages)
+                    {
+                        AssetRegistry.GetAssetsByPackageName(PackageName, RedirectorReferencers);
+                    }
+                    FilterAssetsToValidate(RedirectorReferencers);
+
+                    if (RedirectorReferencers.Num())
+                    {
+                        Msg->AddText(LOCTEXT("InvalidReferences_ReferencedRedirector", "Note that the referencing asset is a redirector which is referenced by the following assets: "));
+                        for (const FAssetData& Referencer : RedirectorReferencers)
+                        {
+                            Msg->AddToken(FAssetDataToken::Create(Referencer));
+                        }
+                    }
+                    else
+                    {
+                        Msg->AddText(LOCTEXT("InvalidReferences_UnreferencedRedirector", "Note that the referencing asset is a redirector which is not referenced by anything else."));
+                    }
+                }
             }
         }
     }
@@ -88,3 +131,5 @@ EDataValidationResult UEditorValidator_PluginAssetReferences::ValidateLoadedAsse
     }
     return GetValidationResult();
 }
+
+#undef LOCTEXT_NAMESPACE 
\ No newline at end of file
