Change log channel

diff --git a/Engine/Plugins/Editor/AssetReferenceRestrictions/Source/AssetReferenceRestrictions/Private/EditorValidator_PluginAssetReferences.cpp b/Engine/Plugins/Editor/AssetReferenceRestrictions/Source/AssetReferenceRestrictions/Private/EditorValidator_PluginAssetReferences.cpp
--- a/Engine/Plugins/Editor/AssetReferenceRestrictions/Source/AssetReferenceRestrictions/Private/EditorValidator_PluginAssetReferences.cpp
+++ b/Engine/Plugins/Editor/AssetReferenceRestrictions/Source/AssetReferenceRestrictions/Private/EditorValidator_PluginAssetReferences.cpp
@@ -3,6 +3,7 @@
 #include "EditorValidator_PluginAssetReferences.h"
 
 #include "Algo/RemoveIf.h"
+#include "AssetReferencingDomains.h"
 #include "AssetReferencingPolicySettings.h"
 #include "AssetReferencingPolicySubsystem.h"
 #include "AssetRegistry/AssetData.h"
@@ -53,7 +54,7 @@ EDataValidationResult UEditorValidator_PluginAssetReferences::ValidateLoadedAsse
             TSharedPtr<IPlugin> Plugin = PluginManager.FindPlugin(PluginName);
             if (Plugin.IsValid() && Plugin->IsEnabled())
             {
-                UE_LOGFMT(LogContentValidation, Display, "Validating asset references from plugin: {PluginName}", Plugin->GetName());
+                UE_LOGFMT(LogAssetReferenceRestrictions, Display, "Validating asset references from plugin: {PluginName}", Plugin->GetName());
                 PluginsToValidate.Add(Plugin.ToSharedRef());
             }
         }
