UE: Add high level config setting to disable AssetReferenceRestrictions (bUseAssetReferenceRestrictions in [/Script/AssetReferenceRestrictions.AssetReferencingPolicySettings] of Game.ini)

#Synced-CL 50460350

diff --git a/Engine/Plugins/Editor/AssetReferenceRestrictions/Source/AssetReferenceRestrictions/Private/AssetReferencingPolicySettings.cpp b/Engine/Plugins/Editor/AssetReferenceRestrictions/Source/AssetReferenceRestrictions/Private/AssetReferencingPolicySettings.cpp
--- a/Engine/Plugins/Editor/AssetReferenceRestrictions/Source/AssetReferenceRestrictions/Private/AssetReferencingPolicySettings.cpp
+++ b/Engine/Plugins/Editor/AssetReferenceRestrictions/Source/AssetReferenceRestrictions/Private/AssetReferencingPolicySettings.cpp
@@ -31,9 +31,10 @@ UAssetReferencingPolicySettings::UAssetReferencingPolicySettings(const FObjectIn
 #if WITH_EDITOR
 void UAssetReferencingPolicySettings::PostEditChangeProperty(FPropertyChangedEvent& PropertyChangedEvent)
 {
-	UAssetReferencingPolicySubsystem* Subsystem = GEditor->GetEditorSubsystem<UAssetReferencingPolicySubsystem>();
-	check(Subsystem);
-	Subsystem->GetDomainDB()->MarkDirty();
+	if (UAssetReferencingPolicySubsystem* Subsystem = GEditor->GetEditorSubsystem<UAssetReferencingPolicySubsystem>())
+	{
+		Subsystem->GetDomainDB()->MarkDirty();
+	}
 
 	Super::PostEditChangeProperty(PropertyChangedEvent);
 }
@@ -55,9 +56,10 @@ TArray<FString> UAssetReferencingPolicySettings::GetListOfDomains(bool bAllowPlu
 
 	if (bAllowPlugins)
 	{
-		UAssetReferencingPolicySubsystem* Subsystem = GEditor->GetEditorSubsystem<UAssetReferencingPolicySubsystem>();
-		check(Subsystem);
-		Result.Append(Subsystem->GetDomainDB()->GetDomainsDefinedByPlugins());
+		if (UAssetReferencingPolicySubsystem* Subsystem = GEditor->GetEditorSubsystem<UAssetReferencingPolicySubsystem>())
+		{
+			Result.Append(Subsystem->GetDomainDB()->GetDomainsDefinedByPlugins());
+		}
 	}
 
 	if (bAllowEngine)
diff --git a/Engine/Plugins/Editor/AssetReferenceRestrictions/Source/AssetReferenceRestrictions/Private/AssetReferencingPolicySubsystem.cpp b/Engine/Plugins/Editor/AssetReferenceRestrictions/Source/AssetReferenceRestrictions/Private/AssetReferencingPolicySubsystem.cpp
--- a/Engine/Plugins/Editor/AssetReferenceRestrictions/Source/AssetReferenceRestrictions/Private/AssetReferencingPolicySubsystem.cpp
+++ b/Engine/Plugins/Editor/AssetReferenceRestrictions/Source/AssetReferenceRestrictions/Private/AssetReferencingPolicySubsystem.cpp
@@ -25,7 +25,7 @@ namespace AssetReferencingPolicySubsystem::Private
 
 bool UAssetReferencingPolicySubsystem::ShouldCreateSubsystem(UObject* Outer) const
 {
-	return true;
+	return GetDefault<UAssetReferencingPolicySettings>()->bUseAssetReferenceRestrictions;
 }
 
 void UAssetReferencingPolicySubsystem::Initialize(FSubsystemCollectionBase& Collection)
diff --git a/Engine/Plugins/Editor/AssetReferenceRestrictions/Source/AssetReferenceRestrictions/Private/AssetValidator_AssetReferenceRestrictions.cpp b/Engine/Plugins/Editor/AssetReferenceRestrictions/Source/AssetReferenceRestrictions/Private/AssetValidator_AssetReferenceRestrictions.cpp
--- a/Engine/Plugins/Editor/AssetReferenceRestrictions/Source/AssetReferenceRestrictions/Private/AssetValidator_AssetReferenceRestrictions.cpp
+++ b/Engine/Plugins/Editor/AssetReferenceRestrictions/Source/AssetReferenceRestrictions/Private/AssetValidator_AssetReferenceRestrictions.cpp
@@ -29,7 +29,10 @@ bool UAssetValidator_AssetReferenceRestrictions::CanValidateAsset_Implementation
 {
 	if (InAsset)
 	{
-		return GEditor->GetEditorSubsystem<UAssetReferencingPolicySubsystem>()->ShouldValidateAssetReferences(AssetData);
+		if (UAssetReferencingPolicySubsystem* Subsystem = GEditor->GetEditorSubsystem<UAssetReferencingPolicySubsystem>())
+		{
+			return Subsystem->ShouldValidateAssetReferences(AssetData);
+		}
 	}
 
 	return false;
@@ -62,6 +65,8 @@ EDataValidationResult UAssetValidator_AssetReferenceRestrictions::ValidateLoaded
 void UAssetValidator_AssetReferenceRestrictions::ValidateAssetInternal(const FAssetData& InAssetData, const IAssetRegistry& InAssetRegistry)
 {
     UAssetReferencingPolicySubsystem* Subsystem = GEditor->GetEditorSubsystem<UAssetReferencingPolicySubsystem>();
+	check(Subsystem);
+
 	if (TValueOrError<void, TArray<FAssetReferenceError>> Result = Subsystem->ValidateAssetReferences(InAssetData, EAssetReferenceFilterRole::Validation);
 		Result.HasError())
 	{
diff --git a/Engine/Plugins/Editor/AssetReferenceRestrictions/Source/AssetReferenceRestrictions/Private/EditorValidator_PluginAssetReferences.cpp b/Engine/Plugins/Editor/AssetReferenceRestrictions/Source/AssetReferenceRestrictions/Private/EditorValidator_PluginAssetReferences.cpp
--- a/Engine/Plugins/Editor/AssetReferenceRestrictions/Source/AssetReferenceRestrictions/Private/EditorValidator_PluginAssetReferences.cpp
+++ b/Engine/Plugins/Editor/AssetReferenceRestrictions/Source/AssetReferenceRestrictions/Private/EditorValidator_PluginAssetReferences.cpp
@@ -23,6 +23,12 @@
 
 bool UEditorValidator_PluginAssetReferences::CanValidateAsset_Implementation(const FAssetData& InAssetData, UObject* InObject, FDataValidationContext& InContext) const
 {
+	if (!GEditor->GetEditorSubsystem<UAssetReferencingPolicySubsystem>())
+	{
+		// AssetReferenceRestrictions is disabled
+		return false;
+	}
+
     if (UDataValidationChangelist* Changelist = Cast<UDataValidationChangelist>(InObject))
     {
         for (const FString& ModifiedFilePath : Changelist->ModifiedFiles)
diff --git a/Engine/Plugins/Editor/AssetReferenceRestrictions/Source/AssetReferenceRestrictions/Public/AssetReferencingPolicySettings.h b/Engine/Plugins/Editor/AssetReferenceRestrictions/Source/AssetReferenceRestrictions/Public/AssetReferencingPolicySettings.h
--- a/Engine/Plugins/Editor/AssetReferenceRestrictions/Source/AssetReferenceRestrictions/Public/AssetReferencingPolicySettings.h
+++ b/Engine/Plugins/Editor/AssetReferenceRestrictions/Source/AssetReferenceRestrictions/Public/AssetReferencingPolicySettings.h
@@ -182,6 +182,10 @@ class UAssetReferencingPolicySettings : public UDeveloperSettings
 	GENERATED_BODY()
 
 public:
+	// High level switch to enable or disable all AssetReferenceRestrictions
+	UPROPERTY(config, EditAnywhere, Category = Settings)
+	bool bUseAssetReferenceRestrictions = true;
+
 	// Settings/rules for engine plugins
 	UPROPERTY(config, EditAnywhere, Category="Engine Plugins")
 	FARPDomainSettingsForPlugins EnginePlugins;
