Added in-editor validation for missing hard-refs

This is disabled by default, but can be enabled by setting the CVar 'AssetReferencingPolicy.CheckForMissingHardRefs' to true

#Synced-CL 52050999

diff --git a/Engine/Plugins/Editor/AssetReferenceRestrictions/Source/AssetReferenceRestrictions/Private/AssetReferencingPolicySubsystem.cpp b/Engine/Plugins/Editor/AssetReferenceRestrictions/Source/AssetReferenceRestrictions/Private/AssetReferencingPolicySubsystem.cpp
--- a/Engine/Plugins/Editor/AssetReferenceRestrictions/Source/AssetReferenceRestrictions/Private/AssetReferencingPolicySubsystem.cpp
+++ b/Engine/Plugins/Editor/AssetReferenceRestrictions/Source/AssetReferenceRestrictions/Private/AssetReferencingPolicySubsystem.cpp
@@ -15,6 +15,13 @@
 
 namespace AssetReferencingPolicySubsystem::Private
 {
+	bool bCheckForMissingHardRefs = false;
+	static FAutoConsoleVariableRef CVarCheckForMissingHardRefs(
+		TEXT("AssetReferencingPolicy.CheckForMissingHardRefs"),
+		bCheckForMissingHardRefs,
+		TEXT("Error on missing hard dependencies when validating asset references.")
+	);
+
 	bool bCheckForMissingRefsToExternalActors = false;
 	static FAutoConsoleVariableRef CVarCheckForMissingRefsToExternalActors(
 		TEXT("AssetReferencingPolicy.CheckForMissingRefsToExternalActors"),
@@ -61,13 +68,12 @@ TValueOrError<void, TArray<FAssetReferenceError>> UAssetReferencingPolicySubsyst
 
 TValueOrError<void, TArray<FAssetReferenceError>> UAssetReferencingPolicySubsystem::ValidateAssetReferences(const FAssetData& InAssetData, const EAssetReferenceFilterRole Role) const
 {
-	IAssetRegistry& AssetRegistry = FModuleManager::LoadModuleChecked<FAssetRegistryModule>("AssetRegistry").Get();
+	IAssetRegistry& AssetRegistry = IAssetRegistry::GetChecked();
 	TArray<FAssetReferenceError> Errors;
 	static const FName TransientName = GetTransientPackage()->GetFName();
 
 	// Check for missing soft or hard references to cinematic and developers content
 	FName PackageFName = InAssetData.PackageName;
-	TArray<FName> SoftDependencies;
 	TArray<FAssetData> AllDependencyAssets;
 
 	const UAssetReferencingPolicySettings* Settings = GetDefault<UAssetReferencingPolicySettings>();
@@ -79,39 +85,39 @@ TValueOrError<void, TArray<FAssetReferenceError>> UAssetReferencingPolicySubsyst
 		QueryFlags = UE::AssetRegistry::EDependencyQuery::Game;
 	}
 
-	AssetRegistry.GetDependencies(PackageFName, SoftDependencies, UE::AssetRegistry::EDependencyCategory::Package, UE::AssetRegistry::EDependencyQuery::Soft | QueryFlags);
-	for (FName SoftDependency : SoftDependencies)
-	{
-		FNameBuilder SoftDependencyStr(SoftDependency);
-		if (!FPackageName::IsScriptPackage(SoftDependencyStr.ToView()) && !FPackageName::IsVersePackage(SoftDependencyStr.ToView()))
+	auto ValidateDependencies = 
+		[PackageFName, &AssetRegistry, &AllDependencyAssets, &Errors](const UE::AssetRegistry::EDependencyQuery DependenciesQueryFlags, const FText& MissingRefMsg, const bool bReportErrors = true)
 		{
-			TArray<FAssetData> DependencyAssets;
-			AssetRegistry.GetAssetsByPackageName(SoftDependency, DependencyAssets, true);
-			if (DependencyAssets.Num() == 0)
+			TArray<FName> PackageDependencies;
+			AssetRegistry.GetDependencies(PackageFName, PackageDependencies, UE::AssetRegistry::EDependencyCategory::Package, DependenciesQueryFlags);
+			for (FName PackageDependency : PackageDependencies)
 			{
-				if (SoftDependency != TransientName)
+				FNameBuilder PackageDependencyStr(PackageDependency);
+				if (!FPackageName::IsScriptPackage(PackageDependencyStr.ToView()) && !FPackageName::IsVersePackage(PackageDependencyStr.ToView()))
 				{
-					Errors.Add(FAssetReferenceError{
-						.Type = EAssetReferenceErrorType::DoesNotExist,
-						.ReferencedAsset = FAssetData(SoftDependency, FName(FPathViews::GetPath(SoftDependencyStr)), FName(), UObject::StaticClass()->GetClassPathName()),
-						.Message = FText::Format(LOCTEXT("IllegalReference_MissingSoftRef", "Soft references {0} which does not exist"), FText::FromName(SoftDependency)),
-					});
+					TArray<FAssetData> DependencyAssets;
+					AssetRegistry.GetAssetsByPackageName(PackageDependency, DependencyAssets, true);
+					if (DependencyAssets.Num() == 0)
+					{
+						if (bReportErrors && PackageDependency != TransientName)
+						{
+							Errors.Add(FAssetReferenceError{
+								.Type = EAssetReferenceErrorType::DoesNotExist,
+								.ReferencedAsset = FAssetData(PackageDependency, FName(FPathViews::GetPath(PackageDependencyStr)), FName(), UObject::StaticClass()->GetClassPathName()),
+								.Message = MissingRefMsg,
+								});
+						}
+					}
+					else
+					{
+						AllDependencyAssets.Append(DependencyAssets);
+					}
 				}
 			}
-			else
-			{
-				AllDependencyAssets.Append(DependencyAssets);
-			}
-		}
-	}
+		};
 
-	// Now check hard references to cinematic and developers content
-	TArray<FName> HardDependencies;
-	AssetRegistry.GetDependencies(PackageFName, HardDependencies, UE::AssetRegistry::EDependencyCategory::Package, UE::AssetRegistry::EDependencyQuery::Hard | QueryFlags);
-	for (FName HardDependency : HardDependencies)
-	{
-		AssetRegistry.GetAssetsByPackageName(HardDependency, AllDependencyAssets, true);
-	}
+	ValidateDependencies(UE::AssetRegistry::EDependencyQuery::Soft | QueryFlags, LOCTEXT("IllegalReference_MissingSoftRef", "soft references a missing package:"));
+	ValidateDependencies(UE::AssetRegistry::EDependencyQuery::Hard | QueryFlags, LOCTEXT("IllegalReference_MissingHardRef", "hard references a missing package (resave to clear this reference):"), AssetReferencingPolicySubsystem::Private::bCheckForMissingHardRefs);
 
 	// Check for missing external actor references which can manifest as LoadErrors and show as either Hard/Soft/Game/EditorOnly.
 	if (AssetReferencingPolicySubsystem::Private::bCheckForMissingRefsToExternalActors)
@@ -130,7 +136,7 @@ TValueOrError<void, TArray<FAssetReferenceError>> UAssetReferencingPolicySubsyst
 				Errors.Add(FAssetReferenceError{
 					.Type = EAssetReferenceErrorType::DoesNotExist,
 					.ReferencedAsset = FAssetData(Dependency, FName(FPathViews::GetPath(DependencyStr)), FName(), UObject::StaticClass()->GetClassPathName()),
-					.Message = FText::Format(LOCTEXT("IllegalReference_MissingRefToExternalActor", "references external actor {0} which does not exist"), FText::FromName(Dependency)),
+					.Message = LOCTEXT("IllegalReference_MissingRefToExternalActor", "references a missing external actor package:"),
 				});
 			}
 		}
