Fix AssetReferenceRestrictions phantom domain error for platform-disabled plugin dependencies

#rnx
#Synced-CL 52165013

diff --git a/Engine/Plugins/Editor/AssetReferenceRestrictions/Source/AssetReferenceRestrictions/Private/AssetReferencingDomains.cpp b/Engine/Plugins/Editor/AssetReferenceRestrictions/Source/AssetReferenceRestrictions/Private/AssetReferencingDomains.cpp
--- a/Engine/Plugins/Editor/AssetReferenceRestrictions/Source/AssetReferenceRestrictions/Private/AssetReferencingDomains.cpp
+++ b/Engine/Plugins/Editor/AssetReferenceRestrictions/Source/AssetReferenceRestrictions/Private/AssetReferencingDomains.cpp
@@ -451,7 +451,7 @@ void FDomainDatabase::BuildDomainFromPlugin(TSharedRef<IPlugin> Plugin)
 	{
 		if (Dependencies.bEnabled)
 		{
-			TSharedPtr<IPlugin> PluginWeDependOn = IPluginManager::Get().FindPlugin(Dependencies.Name);
+			TSharedPtr<IPlugin> PluginWeDependOn = IPluginManager::Get().FindEnabledPlugin(Dependencies.Name);
 			if (PluginWeDependOn.IsValid() && (PluginWeDependOn->CanContainContent() || PluginWeDependOn->CanContainVerse()))
 			{
 				Domain->DomainsVisibleFromHere.Add(FindOrAddDomainByName(Dependencies.Name));
